Copying //Tasks/UE4/Dev-Online-Plugins-v3 to //UE4/Dev-Main (Source: //Tasks/UE4/Dev-Online-Plugins-v3 @ 3047294)

diff --git a/Engine/Plugins/Online/IOS/OnlineSubsystemIOS/Source/Private/OnlineLeaderboardsInterfaceIOS.cpp b/Engine/Plugins/Online/IOS/OnlineSubsystemIOS/Source/Private/OnlineLeaderboardsInterfaceIOS.cpp
new file mode 100644
--- /dev/null
+++ b/Engine/Plugins/Online/IOS/OnlineSubsystemIOS/Source/Private/OnlineLeaderboardsInterfaceIOS.cpp
@@ -0,0 +1,399 @@
+// Copyright 1998-2016 Epic Games, Inc. All Rights Reserved.
+
+// Module includes
+#include "OnlineSubsystemIOSPrivatePCH.h"
+
+// GameCenter includes
+#include <GameKit/GKLeaderboard.h>
+#include <GameKit/GKScore.h>
+
+FOnlineLeaderboardsIOS::FOnlineLeaderboardsIOS(FOnlineSubsystemIOS* InSubsystem)
+{
+	UE_LOG(LogOnline, Display, TEXT("FOnlineLeaderboardsIOS::FOnlineLeaderboardsIOS()"));
+
+	// Cache a reference to the OSS Identity and Friends interfaces, we will need these when we are performing leaderboard actions
+	IdentityInterface = (FOnlineIdentityIOS*)InSubsystem->GetIdentityInterface().Get();
+	check(IdentityInterface);
+
+	FriendsInterface = (FOnlineFriendsIOS*)InSubsystem->GetFriendsInterface().Get();
+	check(FriendsInterface);
+
+	UnreportedScores = nil;
+}
+
+
+FOnlineLeaderboardsIOS::~FOnlineLeaderboardsIOS()
+{
+	if(UnreportedScores)
+	{
+		[UnreportedScores release];
+		UnreportedScores = nil;
+	}
+}
+
+bool FOnlineLeaderboardsIOS::ReadLeaderboardCompletionDelegate(NSArray* players, FOnlineLeaderboardReadRef& InReadObject)
+{
+    auto ReadObject = InReadObject;
+    bool bTriggeredReadRequest = false;
+
+    GKLeaderboard* LeaderboardRequest = nil;
+#ifdef __IPHONE_8_0
+	if ([GKLeaderboard instancesRespondToSelector:@selector(initWithPlayers:)] == YES)
+	{
+        LeaderboardRequest = [[GKLeaderboard alloc] initWithPlayers:players];
+    }
+    else
+#endif
+    {
+#if __IPHONE_OS_VERSION_MIN_REQUIRED < __IPHONE_8_0
+        LeaderboardRequest = [[GKLeaderboard alloc] initWithPlayerIDs:players];
+#endif
+    }
+    if (LeaderboardRequest != nil)
+    {
+        const FString LeaderboardName = ReadObject->LeaderboardName.ToString();
+        
+        NSString* Category = [NSString stringWithFString:LeaderboardName];
+        UE_LOG(LogOnline, Display, TEXT("Attempting to read leaderboard: %s"), *LeaderboardName);
+        
+        LeaderboardRequest.playerScope = GKLeaderboardPlayerScopeGlobal;
+        LeaderboardRequest.timeScope = GKLeaderboardTimeScopeToday;
+#ifdef __IPHONE_7_0
+        if ([LeaderboardRequest respondsToSelector:@selector(identifier)] == YES)
+        {
+            LeaderboardRequest.identifier = Category;
+        }
+        else
+#endif
+        {
+#if __IPHONE_OS_VERSION_MIN_REQUIRED < __IPHONE_7_0
+            LeaderboardRequest.category = Category;
+#endif
+        }
+        LeaderboardRequest.range = NSMakeRange(1,10);
+        
+        bTriggeredReadRequest = true;
+        dispatch_async(dispatch_get_main_queue(), ^
+        {
+            [LeaderboardRequest loadScoresWithCompletionHandler: ^(NSArray *scores, NSError *Error)
+             {
+                bool bWasSuccessful = (Error == nil) && [scores count] > 0;
+                            
+                if (bWasSuccessful)
+                {
+                    bWasSuccessful = [scores count] > 0;
+                    UE_LOG(LogOnline, Display, TEXT("FOnlineLeaderboardsIOS::loadScoresWithCompletionHandler() - %s"), (bWasSuccessful ? TEXT("Success!") : TEXT("Failed!, no scores retrieved")));
+                    for (GKScore* score in scores)
+                    {
+                        FString PlayerIDString;
+                            
+#ifdef __IPHONE_8_0
+						if ([score respondsToSelector:@selector(player)] == YES)
+						{
+                            PlayerIDString = FString(score.player.playerID);
+                        }
+                        else
+#endif
+                        {
+#if __IPHONE_OS_VERSION_MIN_REQUIRED < __IPHONE_8_0
+                            PlayerIDString = FString(score.playerID);
+#endif
+                        }
+                            
+                        UE_LOG(LogOnline, Display, TEXT("----------------------------------------------------------------"));
+                        UE_LOG(LogOnline, Display, TEXT("PlayerId: %s"), *PlayerIDString);
+                        UE_LOG(LogOnline, Display, TEXT("Value: %d"), score.value);
+                        UE_LOG(LogOnline, Display, TEXT("----------------------------------------------------------------"));
+                            
+                        TSharedRef<const FUniqueNetId> UserId = MakeShareable(new FUniqueNetIdString(PlayerIDString));
+                            
+                        FOnlineStatsRow* UserRow = ReadObject.Get().FindPlayerRecord(UserId.Get());
+                        if (UserRow == NULL)
+                        {
+                            UserRow = new (ReadObject->Rows) FOnlineStatsRow(PlayerIDString, UserId);
+                        }
+                            
+                        for (int32 StatIdx = 0; StatIdx < ReadObject->ColumnMetadata.Num(); StatIdx++)
+                        {
+                            const FColumnMetaData& ColumnMeta = ReadObject->ColumnMetadata[StatIdx];
+                            
+                            switch (ColumnMeta.DataType)
+                            {
+                                case EOnlineKeyValuePairDataType::Int32:
+                                {
+                                    int32 Value = score.value;
+                                    UserRow->Columns.Add(ColumnMeta.ColumnName, FVariantData(Value));
+                                    bWasSuccessful = true;
+                                    break;
+                                }
+                            
+                                default:
+                                {
+                                    UE_LOG_ONLINE(Warning, TEXT("Unsupported key value pair during retrieval from GameCenter %s"), ColumnMeta.ColumnName);
+                                    break;
+                                }
+                            }
+                        }
+                    }
+                }
+                else if (Error)
+                {
+                    // if we have failed to read the leaderboard then report this
+                    NSDictionary *userInfo = [Error userInfo];
+                    NSString *errstr = [[userInfo objectForKey : NSUnderlyingErrorKey] localizedDescription];
+                    UE_LOG(LogOnline, Display, TEXT("FOnlineLeaderboardsIOS::loadScoresWithCompletionHandler() - Failed to read leaderboard with error: [%s]"), *FString(errstr));
+                    UE_LOG(LogOnline, Warning, TEXT("You should check that the leaderboard name matches that of one in ITunesConnect"));
+                }
+                            
+                // Report back to the game thread whether this succeeded.
+                [FIOSAsyncTask CreateTaskWithBlock : ^ bool(void)
+                 {
+                    ReadObject->ReadState = bWasSuccessful ? EOnlineAsyncTaskState::Done : EOnlineAsyncTaskState::Failed;
+                    TriggerOnLeaderboardReadCompleteDelegates(bWasSuccessful);
+                    return true;
+                 }];
+            }];
+        });
+    }
+    
+    // If we have failed to kick off a read request, we should still tell whoever is listening.
+    if (!bTriggeredReadRequest)
+    {
+        UE_LOG(LogOnline, Display, TEXT("FOnlineLeaderboardsIOS::loadScoresWithCompletionHandler() - Failed!"));
+        TriggerOnLeaderboardReadCompleteDelegates(false);
+    }
+    
+    return bTriggeredReadRequest;
+}
+
+bool FOnlineLeaderboardsIOS::ReadLeaderboards(const TArray< TSharedRef<const FUniqueNetId> >& Players, FOnlineLeaderboardReadRef& InReadObject)
+{
+	__block FOnlineLeaderboardReadRef ReadObject = InReadObject;
+
+	UE_LOG(LogOnline, Display, TEXT("FOnlineLeaderboardsIOS::ReadLeaderboards()"));
+
+	ReadObject->ReadState = EOnlineAsyncTaskState::Failed;
+	ReadObject->Rows.Empty();
+	
+	if ((IdentityInterface != nullptr) && (IdentityInterface->GetLocalGameCenterUser() != NULL) && IdentityInterface->GetLocalGameCenterUser().isAuthenticated)
+	{
+		ReadObject->ReadState = EOnlineAsyncTaskState::InProgress;
+
+		// Populate a list of id's for our friends which we want to look up.
+		NSMutableArray* FriendIds = [NSMutableArray arrayWithCapacity: (Players.Num() + 1)];
+		
+		// Add the local player to the list of ids to look up.
+		TSharedPtr<const FUniqueNetId> LocalPlayerUID = IdentityInterface->GetUniquePlayerId(0);
+		check(LocalPlayerUID.IsValid());
+
+		FriendIds[0] = [NSString stringWithFString:LocalPlayerUID->ToString()];
+
+		// Add the other requested players
+		for (int32 FriendIdx = 0; FriendIdx < Players.Num(); FriendIdx++)
+		{
+			FriendIds[FriendIdx + 1] = [NSString stringWithFString:Players[FriendIdx]->ToString()];
+		}
+
+		// Kick off a game center read request for the list of users
+#ifdef __IPHONE_8_0
+		if ([GKLeaderboard instancesRespondToSelector:@selector(initWithPlayers:)] == YES)
+		{
+            [GKPlayer loadPlayersForIdentifiers:FriendIds withCompletionHandler:^(NSArray *players, NSError *Error)
+             {
+                bool bWasSuccessful = (Error == nil) && [players count] > 0;
+             
+                if (bWasSuccessful)
+                {
+                    bWasSuccessful = [players count] > 0;
+					ReadLeaderboardCompletionDelegate(players, ReadObject);
+				}
+             }];
+        }
+        else
+#endif
+        {
+#if __IPHONE_OS_VERSION_MIN_REQUIRED < __IPHONE_8_0
+            return ReadLeaderboardCompletionDelegate(FriendIds, InReadObject);
+#else
+			return false;
+#endif
+        }
+	}
+
+	return true;
+}
+
+
+bool FOnlineLeaderboardsIOS::ReadLeaderboardsForFriends(int32 LocalUserNum, FOnlineLeaderboardReadRef& ReadObject)
+{
+	UE_LOG(LogOnline, Display, TEXT("FOnlineLeaderboardsIOS::ReadLeaderboardsForFriends()"));
+	if( IdentityInterface->GetLocalGameCenterUser() != NULL && IdentityInterface->GetLocalGameCenterUser().isAuthenticated )
+	{
+		// Gather the friends from the local players game center friends list and perform a read request for these
+		TArray< TSharedRef<FOnlineFriend> > Friends;
+		FriendsInterface->GetFriendsList( 0, EFriendsLists::ToString(EFriendsLists::Default), Friends );
+
+		TArray< TSharedRef<const FUniqueNetId> > FriendIds;
+		for( int32 Idx = 0; Idx < Friends.Num(); Idx++ )
+		{
+			FriendIds.Add( Friends[ Idx ]->GetUserId() );
+		}
+		ReadLeaderboards( FriendIds, ReadObject );
+	}
+	
+	return true;
+}
+
+
+void FOnlineLeaderboardsIOS::FreeStats(FOnlineLeaderboardRead& ReadObject)
+{
+	UE_LOG(LogOnline, Display, TEXT("FOnlineLeaderboardsIOS::FreeStats()"));
+	// not implemented for gc leaderboards
+}
+
+
+bool FOnlineLeaderboardsIOS::WriteLeaderboards(const FName& SessionName, const FUniqueNetId& Player, FOnlineLeaderboardWrite& WriteObject)
+{
+	UE_LOG(LogOnline, Display, TEXT("FOnlineLeaderboardsIOS::WriteLeaderboards()"));
+	bool bWroteAnyLeaderboard = false;
+
+	// Make sure we have storage space for scores
+	if (UnreportedScores == nil)
+	{
+		UnreportedScores = [[NSMutableArray alloc] initWithCapacity : WriteObject.Properties.Num()];
+	}
+
+	//@TODO: Note: The array of leaderboard names is ignored, because they offer no data.
+	// Instead the stat names are used as the leaderboard names for iOS for now.  This whole API needs rework!
+
+	// Queue up the leaderboard stat writes
+	for (FStatPropertyArray::TConstIterator It(WriteObject.Properties); It; ++It)
+	{
+		// Access the stat and the value.
+		const FVariantData& Stat = It.Value();
+
+		FString LeaderboardName(It.Key().ToString());
+		NSString* Category = [NSString stringWithFString:LeaderboardName];
+
+		// Create a leaderboard score object which should be posted to the [Category] leaderboard.
+        GKScore* Score = nil;
+#ifdef __IPHONE_7_0
+		if ([GKScore instancesRespondToSelector:@selector(initWithLeaderboardIdentifier:)] == YES)
+		{
+            Score = [[GKScore alloc] initWithLeaderboardIdentifier:Category];
+        }
+        else
+#endif
+        {
+#if __IPHONE_OS_VERSION_MIN_REQUIRED < __IPHONE_7_0
+            Score = [[GKScore alloc] initWithCategory:Category];
+#else
+			UE_LOG(LogOnline, Warning, TEXT("FOnlineLeaderboardsIOS::WriteLeaderboards(Leaderboard: %s) Could not intiialize score"), *LeaderboardName);
+			return false;
+#endif
+        }
+        
+		Score.context = 0;
+
+		bool bIsValidScore = false;
+
+		// Setup the score with the value we are writing from the variant type
+		switch (Stat.GetType())
+		{
+			case EOnlineKeyValuePairDataType::Int32:
+			{
+				int32 Value;
+				Stat.GetValue(Value);
+				Score.value = Value;
+					
+				bIsValidScore = true;
+				break;
+			}
+
+			default:
+			{
+				UE_LOG(LogOnline, Warning, TEXT("FOnlineLeaderboardsIOS::WriteLeaderboards(Leaderboard: %s) Invalid data type (only Int32 is currently supported)"), *LeaderboardName);
+				break;
+			}
+		}
+
+		if (bIsValidScore)
+		{
+			int32 Value = Score.value;
+			UE_LOG(LogOnline, Display, TEXT("FOnlineLeaderboardsIOS::WriteLeaderboards() Queued score %d on leaderboard %s"), Value, *LeaderboardName);
+
+			[UnreportedScores addObject : Score];
+			bWroteAnyLeaderboard = true;
+		}
+	}
+	
+	// Return whether any stat was cached.
+	return bWroteAnyLeaderboard;
+}
+
+
+bool FOnlineLeaderboardsIOS::FlushLeaderboards(const FName& SessionName)
+{
+	UE_LOG(LogOnline, Display, TEXT("FOnlineLeaderboardsIOS::FlushLeaderboards()"));
+	bool bBeganFlushingScores = false;
+
+	if ((IdentityInterface->GetLocalGameCenterUser() != NULL) && IdentityInterface->GetLocalGameCenterUser().isAuthenticated)
+	{
+		const int32 UnreportedScoreCount = UnreportedScores.count;
+		bBeganFlushingScores = UnreportedScoreCount > 0;
+
+		if (bBeganFlushingScores)
+		{
+			NSArray *ArrayCopy = [[NSArray alloc] initWithArray:UnreportedScores];
+			
+			[UnreportedScores release];
+			UnreportedScores = nil;
+			
+			dispatch_async(dispatch_get_main_queue(), ^
+			{
+				[GKScore reportScores : ArrayCopy withCompletionHandler : ^ (NSError *error)
+				{
+					// Tell whoever was listening that we have written (or failed to write) to the leaderboard
+					bool bSucceeded = error == NULL;
+					if (bSucceeded)
+					{
+						UE_LOG(LogOnline, Display, TEXT("Flushed %d scores to Game Center"), UnreportedScoreCount);
+					}
+					else
+					{
+						UE_LOG(LogOnline, Display, TEXT("Error while flushing scores (code %d)"), [error code]);
+					}
+
+					[ArrayCopy release];
+
+					// Report back to the game thread whether this succeeded.
+					[FIOSAsyncTask CreateTaskWithBlock : ^ bool(void)
+					{
+					TriggerOnLeaderboardFlushCompleteDelegates(SessionName, bSucceeded);
+						return true;
+					}];
+				}
+			];
+			});
+		}
+	}
+
+	// If we didn't begin writing to the leaderboard we should still notify whoever was listening.
+	if (!bBeganFlushingScores)
+	{
+		TriggerOnLeaderboardFlushCompleteDelegates(SessionName, false);
+		UE_LOG(LogOnline, Display, TEXT("Failed to flush scores to leaderboard"));
+	}
+
+	
+	return bBeganFlushingScores;
+}
+
+
+bool FOnlineLeaderboardsIOS::WriteOnlinePlayerRatings(const FName& SessionName, int32 LeaderboardId, const TArray<FOnlinePlayerScore>& PlayerScores)
+{
+	UE_LOG(LogOnline, Display, TEXT("FOnlineLeaderboardsIOS::WriteOnlinePlayerRatings()"));
+	// not implemented for gc leaderboards
+	
+	return false;
+}
\ No newline at end of file
diff --git a/Engine/Plugins/Online/IOS/OnlineSubsystemIOS/Source/Private/OnlineSessionInterfaceIOS.cpp b/Engine/Plugins/Online/IOS/OnlineSubsystemIOS/Source/Private/OnlineSessionInterfaceIOS.cpp
new file mode 100644
--- /dev/null
+++ b/Engine/Plugins/Online/IOS/OnlineSubsystemIOS/Source/Private/OnlineSessionInterfaceIOS.cpp
@@ -0,0 +1,765 @@
+// Copyright 1998-2016 Epic Games, Inc. All Rights Reserved.
+
+#include "OnlineSubsystemIOSPrivatePCH.h"
+
+#if !PLATFORM_TVOS // @todo tvos: What is up with all this being busted?? Multipeer is gone, but so is older stuff? What's to do??
+
+#if __IPHONE_OS_VERSION_MIN_REQUIRED < __IPHONE_7_0
+@implementation FGameCenterSessionDelegateGK
+@synthesize Session;
+
+- (void)initSessionWithName:(NSString*) sessionName
+{
+	UE_LOG(LogOnline, Display, TEXT("- (void)initSessionWithName:(NSString*) sessionName"));
+    self.Session = [[GKSession alloc] initWithSessionID:sessionName displayName:nil sessionMode:GKSessionModePeer];
+    self.Session.delegate = self;
+    self.Session.disconnectTimeout = 5.0f;
+    self.Session.available = YES;
+}
+
+- (void)shutdownSession
+{
+	UE_LOG(LogOnline, Display, TEXT("- (void)shutdownSession"));
+	[self.Session disconnectFromAllPeers];
+	self.Session.available = NO;
+	self.Session.delegate = nil;
+}
+
+- (bool)sessionsAvailable
+{
+    NSArray* availablePeers = [self.Session peersWithConnectionState:GKPeerStateAvailable];
+    return [availablePeers count] > 0;
+}
+
+-(void)joinSession
+{
+    NSString* PeerID = [self.Session peerID];
+    [self connectToPeer:PeerID];
+}
+
+-(void)session:(GKSession *)session didReceiveConnectionRequestFromPeer:(NSString *)peerID
+{
+	UE_LOG(LogOnline, Display, TEXT("-(void)session:(GKSession *)session didReceiveConnectionRequestFromPeer:(NSString *)peerID - %s"), *FString([session displayNameForPeer:peerID]));
+	[session acceptConnectionFromPeer:peerID error:nil];
+}
+
+- (void)connectToPeer:(NSString *)peerId 
+{
+    [self.Session connectToPeer:peerId withTimeout:10];
+}
+
+- (void)session:(GKSession *)session peer:(NSString *)peerID didChangeState:(GKPeerConnectionState)state
+{
+	const FString PeerId([session displayNameForPeer:peerID]);
+	switch (state)
+	{
+		case GKPeerStateAvailable:
+		{
+			UE_LOG(LogOnline, Display, TEXT("Peer available: %s"), *PeerId);
+
+			[session connectToPeer:peerID withTimeout:5.0f];
+			break;
+		}
+
+		case GKPeerStateUnavailable:
+		{
+			UE_LOG(LogOnline, Display, TEXT("Peer unavailable: %s"), *PeerId);
+			break;
+		}
+
+		case GKPeerStateConnected:
+		{
+			UE_LOG(LogOnline, Display, TEXT("Peer connected: %s"), *PeerId);
+			break;
+		}
+
+		case GKPeerStateDisconnected:
+		{
+			UE_LOG(LogOnline, Display, TEXT("Peer disconnected: %s"), *PeerId);
+			break;
+		}
+
+		case GKPeerStateConnecting:
+		{
+			UE_LOG(LogOnline, Display, TEXT("Peer connecting: %s"), *PeerId);
+			break;
+		}
+	}
+}
+
+
+- (void)session:(GKSession *)session didFailWithError:(NSError *)error
+{
+	UE_LOG(LogOnline, Warning, TEXT("Session failed with error code %i"), [error code]);
+
+	[self shutdownSession];
+}
+
+
+- (void)session:(GKSession *)session connectionWithPeerFailed:(NSString *)peerID withError:(NSError *)error
+{
+	NSString* ErrorString = [NSString stringWithFormat:@"connectionWithPeerFailed - Failed to connect to %@ with error code %d", [session displayNameForPeer:peerID], (uint32)[error code]];
+	const FString ConvertedErrorStr(ErrorString);
+	UE_LOG(LogOnline, Display, TEXT("%s"), *ConvertedErrorStr);
+}
+
+@end
+#endif
+
+#if defined(__IPHONE_7_0)
+@implementation FGameCenterSessionDelegateMC
+@synthesize Session;
+@synthesize PeerID;
+
+- (void)initSessionWithName:(NSString*) sessionName
+{
+    UE_LOG(LogOnline, Display, TEXT("- (void)initSessionWithName:(NSString*) sessionName"));
+    self.PeerID = [[MCPeerID alloc] initWithDisplayName:@""];
+    self.Session = [[MCSession alloc] initWithPeer:self.PeerID];
+    self.Session.delegate = self;
+}
+
+- (void)shutdownSession
+{
+    UE_LOG(LogOnline, Display, TEXT("- (void)shutdownSession"));
+    [self.Session disconnect];
+    self.Session.delegate = nil;
+}
+
+- (bool)sessionsAvailable
+{
+    return NO;
+}
+
+-(void)joinSession
+{
+}
+
+- (void)session:(MCSession *)session didReceiveData:(NSData *)data fromPeer:(MCPeerID *)peerID
+{
+    
+}
+
+- (void)session:(MCSession *)session didStartReceivingResourceWithName:(NSString *)resourceName fromPeer:(MCPeerID *)peerID withProgress:(NSProgress *)progress
+{
+}
+
+- (void)session:(MCSession *)session didFinishReceivingResourceWithName:(NSString *)resourceName fromPeer:(MCPeerID *)peerID atURL:(NSURL *)localURL withError:(NSError *)error
+{
+    
+}
+
+- (void)session:(MCSession *)session didReceiveStream:(NSInputStream *)stream withName:(NSString *)streamName fromPeer:(MCPeerID *)peerID
+{
+    
+}
+
+- (void)session:(MCSession *)session peer:(MCPeerID *)peerID didChangeState:(MCSessionState)state
+{
+    const FString PeerId(PeerID.displayName);
+    switch (state)
+    {
+        case MCSessionStateConnected:
+        {
+            UE_LOG(LogOnline, Display, TEXT("Peer connected: %s"), *PeerId);
+            break;
+        }
+            
+        case MCSessionStateNotConnected:
+        {
+            UE_LOG(LogOnline, Display, TEXT("Peer not connected: %s"), *PeerId);
+            break;
+        }
+    }
+}
+
+@end
+#endif
+
+@implementation FGameCenterSessionDelegate
+#if __IPHONE_OS_VERSION_MIN_REQUIRED < __IPHONE_7_0
+@synthesize SessionGK;
+#endif
+#ifdef __IPHONE_7_0
+@synthesize SessionMC;
+#endif
+
+- (void)initSessionWithName:(NSString*) sessionName
+{
+    UE_LOG(LogOnline, Display, TEXT("- (void)initSessionWithName:(NSString*) sessionName"));
+    // Create the session object
+#ifdef __IPHONE_7_0
+    if ([MCSession class])
+    {
+        self.SessionMC = [FGameCenterSessionDelegateMC alloc];
+        [self.SessionMC initSessionWithName:sessionName];
+    }
+    else
+#endif
+    {
+#if __IPHONE_OS_VERSION_MIN_REQUIRED < __IPHONE_7_0
+        self.SessionGK = [FGameCenterSessionDelegateGK alloc];
+        [self.SessionGK initSessionWithName:sessionName];
+#endif
+    }
+}
+
+- (void)shutdownSession
+{
+    UE_LOG(LogOnline, Display, TEXT("- (void)shutdownSession"));
+#ifdef __IPHONE_7_0
+    if ([MCSession class])
+    {
+        [self.SessionMC shutdownSession];
+    }
+    else
+#endif
+    {
+#if __IPHONE_OS_VERSION_MIN_REQUIRED < __IPHONE_7_0
+        [self.SessionGK shutdownSession];
+#endif
+    }
+}
+
+- (bool)sessionsAvailable
+{
+    UE_LOG(LogOnline, Display, TEXT("- (void)shutdownSession"));
+#ifdef __IPHONE_7_0
+    if ([MCSession class])
+    {
+        return [self.SessionMC sessionsAvailable];
+    }
+    else
+#endif
+    {
+#if __IPHONE_OS_VERSION_MIN_REQUIRED < __IPHONE_7_0
+        return [self.SessionGK sessionsAvailable];
+#endif
+    }
+    return NO;
+}
+
+- (void)joinSession
+{
+    UE_LOG(LogOnline, Display, TEXT("- (void)shutdownSession"));
+#ifdef __IPHONE_7_0
+    if ([MCSession class])
+    {
+        [self.SessionMC joinSession];
+    }
+    else
+#endif
+    {
+#if __IPHONE_OS_VERSION_MIN_REQUIRED < __IPHONE_7_0
+        [self.SessionGK joinSession];
+#endif
+    }
+}
+
+@end
+
+
+
+#endif
+
+
+FOnlineSessionIOS::FOnlineSessionIOS() :
+	IOSSubsystem(NULL),
+	CurrentSessionSearch(NULL)
+{
+	UE_LOG(LogOnline, Display, TEXT("FOnlineSessionIOS::FOnlineSessionIOS()"));
+}
+
+
+FOnlineSessionIOS::FOnlineSessionIOS(FOnlineSubsystemIOS* InSubsystem) :
+	IOSSubsystem(InSubsystem),
+	CurrentSessionSearch(NULL)
+{
+	UE_LOG(LogOnline, Display, TEXT("FOnlineSessionIOS::FOnlineSessionIOS(FOnlineSubsystemIOS* InSubsystem)"));
+}
+
+
+FOnlineSessionIOS::~FOnlineSessionIOS()
+{
+
+}
+
+
+void FOnlineSessionIOS::Tick(float DeltaTime)
+{
+
+}
+
+
+FNamedOnlineSession* FOnlineSessionIOS::AddNamedSession(FName SessionName, const FOnlineSessionSettings& SessionSettings)
+{
+	FScopeLock ScopeLock(&SessionLock);
+	return new (Sessions) FNamedOnlineSession(SessionName, SessionSettings);
+}
+
+
+FNamedOnlineSession* FOnlineSessionIOS::AddNamedSession(FName SessionName, const FOnlineSession& Session)
+{
+	FScopeLock ScopeLock(&SessionLock);
+	return new (Sessions) FNamedOnlineSession(SessionName, Session);
+}
+
+
+FNamedOnlineSession* FOnlineSessionIOS::GetNamedSession(FName SessionName)
+{
+	FScopeLock ScopeLock(&SessionLock);
+	for (int32 SearchIndex = 0; SearchIndex < Sessions.Num(); SearchIndex++)
+	{
+		if (Sessions[SearchIndex].SessionName == SessionName)
+		{
+			return &Sessions[SearchIndex];
+		}
+	}
+	return NULL;
+}
+
+
+void FOnlineSessionIOS::RemoveNamedSession(FName SessionName)
+{
+	FScopeLock ScopeLock(&SessionLock);
+	for (int32 SearchIndex = 0; SearchIndex < Sessions.Num(); SearchIndex++)
+	{
+		if (Sessions[SearchIndex].SessionName == SessionName)
+		{
+			Sessions.RemoveAtSwap(SearchIndex);
+			return;
+		}
+	}
+}
+
+
+EOnlineSessionState::Type FOnlineSessionIOS::GetSessionState(FName SessionName) const
+{
+	FScopeLock ScopeLock(&SessionLock);
+	for (int32 SearchIndex = 0; SearchIndex < Sessions.Num(); SearchIndex++)
+	{
+		if (Sessions[SearchIndex].SessionName == SessionName)
+		{
+			return Sessions[SearchIndex].SessionState;
+		}
+	}
+
+	return EOnlineSessionState::NoSession;
+}
+
+
+bool FOnlineSessionIOS::HasPresenceSession()
+{
+	FScopeLock ScopeLock(&SessionLock);
+	for (int32 SearchIndex = 0; SearchIndex < Sessions.Num(); SearchIndex++)
+	{
+		if (Sessions[SearchIndex].SessionSettings.bUsesPresence)
+		{
+			return true;
+		}
+	}
+		
+	return false;
+}
+
+
+bool FOnlineSessionIOS::CreateSession(int32 HostingPlayerNum, FName SessionName, const FOnlineSessionSettings& NewSessionSettings)
+{
+	bool bSuccessfullyCreatedSession = false;
+#if !PLATFORM_TVOS
+	UE_LOG(LogOnline, Display, TEXT("FOnlineSessionIOS::CreateSession"));
+	
+	// Check for an existing session
+	FNamedOnlineSession* Session = GetNamedSession(SessionName);
+	if (Session == NULL)
+	{
+		Session = AddNamedSession(SessionName, NewSessionSettings);
+		Session->SessionState = EOnlineSessionState::Pending;
+		UE_LOG(LogOnline, Display, TEXT("Creating new session."));
+
+		// Create the session object
+        FGameCenterSessionDelegate* NewGKSession = [FGameCenterSessionDelegate alloc];
+		if( NewGKSession != NULL )
+		{
+			UE_LOG(LogOnline, Display, TEXT("Created session delegate"));
+			NSString* SafeSessionName = [NSString stringWithFString:SessionName.ToString()];
+			GKSessions.Add(SessionName, NewGKSession);
+		}
+		else
+		{
+			UE_LOG(LogOnline, Warning, TEXT("Failed to create session delegate"));
+		}
+		
+		bSuccessfullyCreatedSession = true;
+	}
+	else
+	{
+		UE_LOG(LogOnline, Display, TEXT("Cannot create session '%s': session already exists."), *SessionName.ToString());
+	}
+	
+	UE_LOG(LogOnline, Display, TEXT("TriggerOnCreateSessionCompleteDelegates: %s, %d"), *SessionName.ToString(), bSuccessfullyCreatedSession);
+	TriggerOnCreateSessionCompleteDelegates(SessionName, bSuccessfullyCreatedSession);
+#endif
+	
+	return bSuccessfullyCreatedSession;
+}
+
+
+bool FOnlineSessionIOS::CreateSession(const FUniqueNetId& HostingPlayerId, FName SessionName, const FOnlineSessionSettings& NewSessionSettings)
+{
+	// HostingPlayerNum is unused, can pass in anything
+	return CreateSession(0, SessionName, NewSessionSettings);
+}
+
+
+bool FOnlineSessionIOS::StartSession(FName SessionName)
+{
+	bool bSuccessfullyStartedSession = false;
+#if !PLATFORM_TVOS
+	UE_LOG(LogOnline, Display, TEXT("FOnlineSessionIOS::StartSession"));
+		
+	// Check for an existing session
+	FNamedOnlineSession* Session = GetNamedSession(SessionName);
+	if (Session != NULL)
+	{
+		// Find the linked GK session and start it.
+        FGameCenterSessionDelegate* LinkedGKSession = *GKSessions.Find( SessionName );
+		NSString* SafeSessionName = [NSString stringWithFString:SessionName.ToString()];
+		[LinkedGKSession initSessionWithName:SafeSessionName];
+
+		// Update the session state as we are now running.
+		Session->SessionState = EOnlineSessionState::InProgress;
+		bSuccessfullyStartedSession = true;
+	}
+	else
+	{
+		UE_LOG(LogOnline, Warning, TEXT("Failed to create session delegate"));
+	}
+
+	TriggerOnStartSessionCompleteDelegates(SessionName, bSuccessfullyStartedSession);
+#endif
+	return bSuccessfullyStartedSession;
+}
+
+
+bool FOnlineSessionIOS::UpdateSession(FName SessionName, FOnlineSessionSettings& UpdatedSessionSettings, bool bShouldRefreshOnlineData)
+{
+	bool bSuccessfullyUpdatedSession = false;
+	
+	UE_LOG(LogOnline, Display, TEXT("FOnlineSessionIOS::UpdateSession - not implemented"));
+	
+	TriggerOnUpdateSessionCompleteDelegates(SessionName, bSuccessfullyUpdatedSession);
+	
+	return bSuccessfullyUpdatedSession;
+}
+
+
+bool FOnlineSessionIOS::EndSession(FName SessionName)
+{
+	bool bSuccessfullyEndedSession = false;
+	
+	UE_LOG(LogOnline, Display, TEXT("FOnlineSessionIOS::EndSession - not implemented"));
+
+	TriggerOnEndSessionCompleteDelegates(SessionName, bSuccessfullyEndedSession);
+
+	return bSuccessfullyEndedSession;
+}
+
+
+bool FOnlineSessionIOS::DestroySession(FName SessionName, const FOnDestroySessionCompleteDelegate& CompletionDelegate)
+{
+	bool bSuccessfullyDestroyedSession = false;
+#if !PLATFORM_TVOS
+	
+	FNamedOnlineSession* Session = GetNamedSession(SessionName);
+	if (Session == NULL)
+	{
+		FGameCenterSessionDelegate* ExistingGKSession = *GKSessions.Find( SessionName );
+		[ExistingGKSession shutdownSession];
+		
+		// The session info is no longer needed
+		RemoveNamedSession( Session->SessionName );
+
+		GKSessions.Remove( SessionName );
+
+		bSuccessfullyDestroyedSession = true;
+	}
+
+	CompletionDelegate.ExecuteIfBound(SessionName, bSuccessfullyDestroyedSession);
+	TriggerOnDestroySessionCompleteDelegates(SessionName, bSuccessfullyDestroyedSession);
+#endif
+	return bSuccessfullyDestroyedSession;
+}
+
+bool FOnlineSessionIOS::IsPlayerInSession(FName SessionName, const FUniqueNetId& UniqueId)
+{
+	return IsPlayerInSessionImpl(this, SessionName, UniqueId);
+}
+
+bool FOnlineSessionIOS::StartMatchmaking(const TArray< TSharedRef<const FUniqueNetId> >& LocalPlayers, FName SessionName, const FOnlineSessionSettings& NewSessionSettings, TSharedRef<FOnlineSessionSearch>& SearchSettings)
+{
+	UE_LOG(LogOnline, Warning, TEXT("StartMatchmaking is not supported on this platform. Use FindSessions or FindSessionById."));
+	TriggerOnMatchmakingCompleteDelegates(SessionName, false);
+	return false;
+}
+
+
+bool FOnlineSessionIOS::CancelMatchmaking(int32 SearchingPlayerNum, FName SessionName)
+{
+	UE_LOG(LogOnline, Warning, TEXT("CancelMatchmaking is not supported on this platform. Use CancelFindSessions."));
+	TriggerOnCancelMatchmakingCompleteDelegates(SessionName, false);
+	return false;
+}
+
+
+bool FOnlineSessionIOS::CancelMatchmaking(const FUniqueNetId& SearchingPlayerId, FName SessionName)
+{
+	UE_LOG(LogOnline, Warning, TEXT("CancelMatchmaking is not supported on this platform. Use CancelFindSessions."));
+	TriggerOnCancelMatchmakingCompleteDelegates(SessionName, false);
+	return false;
+}
+
+bool FOnlineSessionIOS::FindSessions(int32 SearchingPlayerNum, const TSharedRef<FOnlineSessionSearch>& SearchSettings)
+{
+	bool bSuccessfullyFoundSessions = false;
+
+#if !PLATFORM_TVOS
+	// Don't start another search while one is in progress
+	if (!CurrentSessionSearch.IsValid() && SearchSettings->SearchState != EOnlineAsyncTaskState::InProgress)
+	{
+		for (TMap< FName, FGameCenterSessionDelegate* >::TConstIterator SessionIt(GKSessions); SessionIt; ++SessionIt)
+		{
+			FGameCenterSessionDelegate* GKSession = SessionIt.Value();
+            bSuccessfullyFoundSessions = [GKSession sessionsAvailable];
+		}
+	}
+
+	TriggerOnFindSessionsCompleteDelegates(bSuccessfullyFoundSessions);
+#endif
+	return bSuccessfullyFoundSessions;
+}
+
+
+bool FOnlineSessionIOS::FindSessions(const FUniqueNetId& SearchingPlayerId, const TSharedRef<FOnlineSessionSearch>& SearchSettings)
+{
+	// SearchingPlayerNum is unused, can pass in anything
+	return FindSessions(0, SearchSettings);
+}
+
+bool FOnlineSessionIOS::FindSessionById(const FUniqueNetId& SearchingUserId, const FUniqueNetId& SessionId, const FUniqueNetId& FriendId, const FOnSingleSessionResultCompleteDelegate& CompletionDelegates)
+{
+	FOnlineSessionSearchResult EmptyResult;
+	CompletionDelegates.ExecuteIfBound(0, false, EmptyResult);
+	return true;
+}
+
+bool FOnlineSessionIOS::CancelFindSessions()
+{
+	bool bSuccessfullyCancelledSession = false;
+
+	UE_LOG(LogOnline, Display, TEXT("FOnlineSessionIOS::CancelSession - not implemented"));
+	
+	TriggerOnCancelFindSessionsCompleteDelegates(true);
+
+	return bSuccessfullyCancelledSession;
+}
+
+
+bool FOnlineSessionIOS::PingSearchResults(const FOnlineSessionSearchResult& SearchResult)
+{
+	bool bSuccessfullyPingedSearchResults = false;
+
+	UE_LOG(LogOnline, Display, TEXT("FOnlineSessionIOS::PingSearchResults - not implemented"));
+	
+	return bSuccessfullyPingedSearchResults;
+}
+
+
+bool FOnlineSessionIOS::JoinSession(int32 PlayerNum, FName SessionName, const FOnlineSessionSearchResult& DesiredSession)
+{
+	EOnJoinSessionCompleteResult::Type JoinSessionResult = EOnJoinSessionCompleteResult::UnknownError;
+#if !PLATFORM_TVOS
+	UE_LOG(LogOnline, Display, TEXT("FOnlineSessionIOS::JoinSession"));
+
+	FGameCenterSessionDelegate* SessionDelegate = *GKSessions.Find( SessionName );
+	if( SessionDelegate != NULL )
+	{
+        [SessionDelegate joinSession];
+		JoinSessionResult = EOnJoinSessionCompleteResult::Success;
+	}
+
+	TriggerOnJoinSessionCompleteDelegates(SessionName, JoinSessionResult);
+#endif
+	return JoinSessionResult == EOnJoinSessionCompleteResult::Success;
+}
+
+
+bool FOnlineSessionIOS::JoinSession(const FUniqueNetId& PlayerId, FName SessionName, const FOnlineSessionSearchResult& DesiredSession)
+{
+	// PlayerNum is unused, can pass in anything
+	return JoinSession(0, SessionName, DesiredSession);
+}
+
+
+bool FOnlineSessionIOS::FindFriendSession(int32 LocalUserNum, const FUniqueNetId& Friend)
+{
+	bool bSuccessfullyJointFriendSession = false;
+	
+	UE_LOG(LogOnline, Display, TEXT("FOnlineSessionIOS::FindFriendSession - not implemented"));
+	
+	FOnlineSessionSearchResult EmptyResult;
+	TriggerOnFindFriendSessionCompleteDelegates(LocalUserNum, bSuccessfullyJointFriendSession, EmptyResult);
+
+	return bSuccessfullyJointFriendSession;
+}
+
+
+bool FOnlineSessionIOS::FindFriendSession(const FUniqueNetId& LocalUserId, const FUniqueNetId& Friend)
+{
+	return FindFriendSession(0, Friend);
+}
+
+
+bool FOnlineSessionIOS::SendSessionInviteToFriend(int32 LocalUserNum, FName SessionName, const FUniqueNetId& Friend)
+{
+	bool bSuccessfullySentSessionInviteToFriend = false;
+	
+	UE_LOG(LogOnline, Display, TEXT("FOnlineSessionIOS::SendSessionInviteToFriend - not implemented"));
+	
+	return bSuccessfullySentSessionInviteToFriend;
+}
+
+
+bool FOnlineSessionIOS::SendSessionInviteToFriend(const FUniqueNetId& LocalUserId, FName SessionName, const FUniqueNetId& Friend)
+{
+	UE_LOG(LogOnline, Display, TEXT("FOnlineSessionIOS::SendSessionInviteToFriend - not implemented"));
+	
+	return false;
+}
+
+
+bool FOnlineSessionIOS::SendSessionInviteToFriends(int32 LocalUserNum, FName SessionName, const TArray< TSharedRef<const FUniqueNetId> >& Friends)
+{
+	bool bSuccessfullySentSessionInviteToFriends = false;
+
+	UE_LOG(LogOnline, Display, TEXT("FOnlineSessionIOS::SendSessionInviteToFriends - not implemented"));
+	
+	return bSuccessfullySentSessionInviteToFriends;
+}
+
+
+bool FOnlineSessionIOS::SendSessionInviteToFriends(const FUniqueNetId& LocalUserId, FName SessionName, const TArray< TSharedRef<const FUniqueNetId> >& Friends)
+{
+	UE_LOG(LogOnline, Display, TEXT("FOnlineSessionIOS::SendSessionInviteToFriends - not implemented"));
+	
+	return false;
+}
+
+
+bool FOnlineSessionIOS::GetResolvedConnectString(FName SessionName, FString& ConnectInfo)
+{
+	bool bSuccessfullyGotResolvedConnectString = false;
+	
+	UE_LOG(LogOnline, Display, TEXT("FOnlineSessionIOS::GetResolvedConnectString - not implemented"));
+	
+	return bSuccessfullyGotResolvedConnectString;
+}
+
+
+bool FOnlineSessionIOS::GetResolvedConnectString(const class FOnlineSessionSearchResult& SearchResult, FName PortType, FString& ConnectInfo)
+{
+	return false;
+}
+
+
+FOnlineSessionSettings* FOnlineSessionIOS::GetSessionSettings(FName SessionName)
+{
+	FNamedOnlineSession* Session = GetNamedSession(SessionName);
+	if (Session)
+	{
+		return &Session->SessionSettings;
+	}
+	return NULL;
+}
+
+
+bool FOnlineSessionIOS::RegisterPlayer(FName SessionName, const FUniqueNetId& PlayerId, bool bWasInvited)
+{
+	bool bSuccessfullyRegisteredPlayer = false;
+	
+	UE_LOG(LogOnline, Display, TEXT("FOnlineSessionIOS::RegisterPlayer - not implemented"));	
+	
+	TArray< TSharedRef<const FUniqueNetId> > Players;
+	Players.Add(MakeShareable(new FUniqueNetIdString(PlayerId)));
+	
+	bSuccessfullyRegisteredPlayer = RegisterPlayers(SessionName, Players, bWasInvited);
+
+	return bSuccessfullyRegisteredPlayer;
+}
+
+
+bool FOnlineSessionIOS::RegisterPlayers(FName SessionName, const TArray< TSharedRef<const FUniqueNetId> >& Players, bool bWasInvited)
+{
+	bool bSuccessfullyRegisteredPlayers = false;
+	
+	UE_LOG(LogOnline, Display, TEXT("FOnlineSessionIOS::RegisterPlayers - not implemented"));
+	
+	for (int32 PlayerIdx=0; PlayerIdx < Players.Num(); PlayerIdx++)
+	{
+		TriggerOnRegisterPlayersCompleteDelegates(SessionName, Players, bSuccessfullyRegisteredPlayers);
+	}
+
+	return bSuccessfullyRegisteredPlayers;
+}
+
+
+bool FOnlineSessionIOS::UnregisterPlayer(FName SessionName, const FUniqueNetId& PlayerId)
+{
+	bool bSuccessfullyUnregisteredPlayer = false;
+	
+	UE_LOG(LogOnline, Display, TEXT("FOnlineSessionIOS::UnregisterPlayer - not implemented"));
+	
+	TArray< TSharedRef<const FUniqueNetId> > Players;
+	Players.Add(MakeShareable(new FUniqueNetIdString(PlayerId)));
+	bSuccessfullyUnregisteredPlayer = UnregisterPlayers(SessionName, Players);
+	
+	return bSuccessfullyUnregisteredPlayer;
+}
+
+
+bool FOnlineSessionIOS::UnregisterPlayers(FName SessionName, const TArray< TSharedRef<const FUniqueNetId> >& Players)
+{
+	bool bSuccessfullyUnregisteredPlayers = false;
+
+	UE_LOG(LogOnline, Display, TEXT("FOnlineSessionIOS::UnregisterPlayers - not implemented"));
+
+	for (int32 PlayerIdx=0; PlayerIdx < Players.Num(); PlayerIdx++)
+	{
+		TriggerOnUnregisterPlayersCompleteDelegates(SessionName, Players, bSuccessfullyUnregisteredPlayers);
+	}
+
+	return bSuccessfullyUnregisteredPlayers;
+}
+
+
+int32 FOnlineSessionIOS::GetNumSessions()
+{
+	FScopeLock ScopeLock(&SessionLock);
+	return Sessions.Num();
+}
+
+
+void FOnlineSessionIOS::DumpSessionState()
+{
+
+}
+
+void FOnlineSessionIOS::RegisterLocalPlayer(const FUniqueNetId& PlayerId, FName SessionName, const FOnRegisterLocalPlayerCompleteDelegate& Delegate)
+{
+	Delegate.ExecuteIfBound(PlayerId, EOnJoinSessionCompleteResult::Success);
+}
+
+void FOnlineSessionIOS::UnregisterLocalPlayer(const FUniqueNetId& PlayerId, FName SessionName, const FOnUnregisterLocalPlayerCompleteDelegate& Delegate)
+{
+	Delegate.ExecuteIfBound(PlayerId, true);
+}
diff --git a/Engine/Plugins/Online/IOS/OnlineSubsystemIOS/Source/Private/OnlineStoreInterfaceIOS.cpp b/Engine/Plugins/Online/IOS/OnlineSubsystemIOS/Source/Private/OnlineStoreInterfaceIOS.cpp
new file mode 100644
--- /dev/null
+++ b/Engine/Plugins/Online/IOS/OnlineSubsystemIOS/Source/Private/OnlineStoreInterfaceIOS.cpp
@@ -0,0 +1,614 @@
+// Copyright 1998-2016 Epic Games, Inc. All Rights Reserved.
+
+#include "OnlineSubsystemIOSPrivatePCH.h"
+#include "OnlineStoreInterface.h"
+#include "OnlineSubsystemNames.h"
+
+////////////////////////////////////////////////////////////////////
+/// FStoreKitHelper implementation
+
+
+@implementation FStoreKitHelper
+@synthesize Request;
+@synthesize AvailableProducts;
+
+
+- (id)init
+{
+	self = [super init];
+	return self;
+}
+
+
+- (void)paymentQueue:(SKPaymentQueue *)queue updatedTransactions:(NSArray *)transactions
+{
+	UE_LOG(LogOnline, Log, TEXT( "FStoreKitHelper::updatedTransactions" ));
+    for (SKPaymentTransaction *transaction in transactions)
+    {
+        switch ([transaction transactionState])
+        {
+            case SKPaymentTransactionStatePurchased:
+                [self completeTransaction:transaction];
+                break;
+            case SKPaymentTransactionStateFailed:
+                [self failedTransaction:transaction];
+                break;
+            case SKPaymentTransactionStateRestored:
+                [self restoreTransaction:transaction];
+				break;
+            case SKPaymentTransactionStatePurchasing:
+                continue;
+			default:
+				break;
+        }
+    }
+}
+
+-(void)paymentQueueRestoreCompletedTransactionsFinished:(SKPaymentQueue *)queue
+{
+	UE_LOG(LogOnline, Log, TEXT("FStoreKitHelper::paymentQueueRestoreCompletedTransactionsFinished"));
+
+	[FIOSAsyncTask CreateTaskWithBlock : ^ bool(void)
+	{
+		IOnlineSubsystem* OnlineSub = IOnlineSubsystem::Get(IOS_SUBSYSTEM);
+		FOnlineStoreInterfaceIOS* StoreInterface = (FOnlineStoreInterfaceIOS*)OnlineSub->GetStoreInterface().Get();
+
+		if (StoreInterface->CachedPurchaseRestoreObject.IsValid())
+		{
+			StoreInterface->CachedPurchaseRestoreObject->ReadState = EOnlineAsyncTaskState::Done;
+		}
+		StoreInterface->TriggerOnInAppPurchaseRestoreCompleteDelegates(EInAppPurchaseState::Restored);
+		// Ensure these flags are false to allow a subsequent transaction attempt to proceed
+		StoreInterface->bIsProductRequestInFlight = false;
+		StoreInterface->bIsPurchasing = false;
+		StoreInterface->bIsRestoringPurchases = false;
+
+		return true;
+	}];
+}
+
+-(void)paymentQueue: (SKPaymentQueue *)queue restoreCompletedTransactionsFailedWithError : (NSError *)error
+{
+	UE_LOG(LogOnline, Log, TEXT("FStoreKitHelper::restoreCompletedTransactionsFailedWithError - %s"), *FString([error localizedDescription]));
+
+	EInAppPurchaseState::Type CompletionState = EInAppPurchaseState::Unknown;
+	switch (error.code)
+	{
+	case SKErrorPaymentCancelled:
+		CompletionState = EInAppPurchaseState::Cancelled;
+		break;
+	case SKErrorClientInvalid:
+	case SKErrorStoreProductNotAvailable:
+	case SKErrorPaymentInvalid:
+		CompletionState = EInAppPurchaseState::Invalid;
+		break;
+	case SKErrorPaymentNotAllowed:
+		CompletionState = EInAppPurchaseState::NotAllowed;
+		break;
+	}
+
+	[FIOSAsyncTask CreateTaskWithBlock : ^ bool(void)
+	{
+		IOnlineSubsystem* OnlineSub = IOnlineSubsystem::Get(IOS_SUBSYSTEM);
+		FOnlineStoreInterfaceIOS* StoreInterface = (FOnlineStoreInterfaceIOS*)OnlineSub->GetStoreInterface().Get();
+		if (StoreInterface->CachedPurchaseRestoreObject.IsValid())
+		{
+			StoreInterface->CachedPurchaseRestoreObject->ReadState = EOnlineAsyncTaskState::Done;
+		}
+
+		StoreInterface->TriggerOnInAppPurchaseRestoreCompleteDelegates(CompletionState);
+		// Ensure these flags are false to allow a subsequent transaction attempt to proceed
+		StoreInterface->bIsProductRequestInFlight = false;
+		StoreInterface->bIsPurchasing = false;
+		StoreInterface->bIsRestoringPurchases = false;
+
+		return true;
+	}];
+}
+
+-(void)paymentQueue: (SKPaymentQueue *)queue requestProductDataFailedWithError : (NSError *)error
+{
+	UE_LOG(LogOnline, Log, TEXT("FStoreKitHelper::requestProductDataFailedWithError - %s"), *FString([error localizedDescription]));
+
+	[FIOSAsyncTask CreateTaskWithBlock : ^ bool(void)
+	{
+		IOnlineSubsystem* OnlineSub = IOnlineSubsystem::Get(IOS_SUBSYSTEM);
+		FOnlineStoreInterfaceIOS* StoreInterface = (FOnlineStoreInterfaceIOS*)OnlineSub->GetStoreInterface().Get();
+		
+		StoreInterface->TriggerOnQueryForAvailablePurchasesCompleteDelegates( false );
+		// Ensure these flags are false to allow a subsequent transaction attempt to proceed
+		StoreInterface->bIsProductRequestInFlight = false;
+		StoreInterface->bIsPurchasing = false;
+		StoreInterface->bIsRestoringPurchases = false;
+
+		return true;
+	}];
+}
+
+- (void) completeTransaction: (SKPaymentTransaction *)transaction
+{
+	UE_LOG(LogOnline, Log, TEXT( "FStoreKitHelper::completeTransaction" ));
+
+    
+	[FIOSAsyncTask CreateTaskWithBlock : ^ bool(void)
+	{
+		IOnlineSubsystem* OnlineSub = IOnlineSubsystem::Get(IOS_SUBSYSTEM);
+		FOnlineStoreInterfaceIOS* StoreInterface = (FOnlineStoreInterfaceIOS*)OnlineSub->GetStoreInterface().Get();
+
+		if (StoreInterface->CachedPurchaseStateObject.IsValid())
+		{
+			FString ReceiptData;
+
+#ifdef __IPHONE_7_0
+			if( [IOSAppDelegate GetDelegate].OSVersion >= 7.0 )
+			{
+				NSURL* nsReceiptUrl = [[NSBundle mainBundle] appStoreReceiptURL];
+				NSData* nsReceiptData = [NSData dataWithContentsOfURL : nsReceiptUrl];
+				NSString* nsEncodedReceiptData = [nsReceiptData base64EncodedStringWithOptions : NSDataBase64EncodingEndLineWithLineFeed];
+     
+				ReceiptData = nsEncodedReceiptData;
+			}
+			else
+#endif
+			{
+#if __IPHONE_OS_VERSION_MIN_REQUIRED < __IPHONE_7_0
+				// If earlier than IOS 7, we will need to use the transactionReceipt
+				NSString* nsEncodedReceiptData = [transaction.transactionReceipt base64EncodedStringWithOptions:NSDataBase64EncodingEndLineWithLineFeed];
+				ReceiptData = nsEncodedReceiptData;
+#endif
+			}
+
+			StoreInterface->CachedPurchaseStateObject->ProvidedProductInformation.ReceiptData = ReceiptData;
+			StoreInterface->CachedPurchaseStateObject->ProvidedProductInformation.TransactionIdentifier = transaction.transactionIdentifier;
+			StoreInterface->CachedPurchaseStateObject->ReadState = EOnlineAsyncTaskState::Done;
+		}
+     
+		StoreInterface->TriggerOnInAppPurchaseCompleteDelegates(EInAppPurchaseState::Success);
+		// Ensure these flags are false to allow a subsequent transaction attempt to proceed
+ 		StoreInterface->bIsProductRequestInFlight = false;
+		StoreInterface->bIsPurchasing = false;
+		StoreInterface->bIsRestoringPurchases = false;
+	 
+		return true;
+	}];
+ 
+	// Remove the transaction from the payment queue.
+	[[SKPaymentQueue defaultQueue] finishTransaction: transaction];
+}
+
+
+- (void) restoreTransaction: (SKPaymentTransaction *)transaction
+{
+	UE_LOG(LogOnline, Log, TEXT( "FStoreKitHelper::restoreTransaction" ));
+
+	[FIOSAsyncTask CreateTaskWithBlock : ^ bool(void)
+	{
+		IOnlineSubsystem* OnlineSub = IOnlineSubsystem::Get(IOS_SUBSYSTEM);
+		FOnlineStoreInterfaceIOS* StoreInterface = (FOnlineStoreInterfaceIOS*)OnlineSub->GetStoreInterface().Get();
+
+		if (StoreInterface->CachedPurchaseRestoreObject.IsValid())
+		{
+			FString ReceiptData;
+
+#ifdef __IPHONE_7_0
+			if ([IOSAppDelegate GetDelegate].OSVersion >= 7.0)
+			{
+				NSURL* nsReceiptUrl = [[NSBundle mainBundle] appStoreReceiptURL];
+				NSData* nsReceiptData = [NSData dataWithContentsOfURL : nsReceiptUrl];
+				NSString* nsEncodedReceiptData = [nsReceiptData base64EncodedStringWithOptions : NSDataBase64EncodingEndLineWithLineFeed];
+
+				ReceiptData = nsEncodedReceiptData;
+			}
+			else
+#endif
+			{
+#if __IPHONE_OS_VERSION_MIN_REQUIRED < __IPHONE_7_0
+				// If earlier than IOS 7, we will need to use the transactionReceipt
+				NSString* nsEncodedReceiptData = [transaction.originalTransaction.transactionReceipt base64EncodedStringWithOptions : NSDataBase64EncodingEndLineWithLineFeed];
+				ReceiptData = nsEncodedReceiptData;
+#endif
+			}
+
+			FInAppPurchaseRestoreInfo RestoreInfo;
+			RestoreInfo.Identifier = FString(transaction.originalTransaction.payment.productIdentifier);
+			RestoreInfo.ReceiptData = ReceiptData;
+			StoreInterface->CachedPurchaseRestoreObject->ProvidedRestoreInformation.Add(RestoreInfo);
+		}
+
+		return true;
+	}];
+	
+	[[SKPaymentQueue defaultQueue] finishTransaction:transaction];
+}
+
+
+- (void) failedTransaction: (SKPaymentTransaction *)transaction
+{
+	UE_LOG(LogOnline, Log, TEXT( "FStoreKitHelper::failedTransaction - %s" ), *FString([transaction.error localizedDescription]));
+
+	EInAppPurchaseState::Type CompletionState = EInAppPurchaseState::Unknown;
+	switch (transaction.error.code)
+	{
+	case SKErrorPaymentCancelled:
+		CompletionState = EInAppPurchaseState::Cancelled;
+		break;
+        case SKErrorClientInvalid:
+        case SKErrorStoreProductNotAvailable:
+        case SKErrorPaymentInvalid:
+		CompletionState = EInAppPurchaseState::Invalid;
+		break;
+	case SKErrorPaymentNotAllowed:
+		CompletionState = EInAppPurchaseState::NotAllowed;
+		break;
+	}
+	
+    [FIOSAsyncTask CreateTaskWithBlock : ^ bool(void)
+    {
+        IOnlineSubsystem* OnlineSub = IOnlineSubsystem::Get(IOS_SUBSYSTEM);
+        FOnlineStoreInterfaceIOS* StoreInterface = (FOnlineStoreInterfaceIOS*)OnlineSub->GetStoreInterface().Get();
+        if (StoreInterface->CachedPurchaseStateObject.IsValid())
+        {
+            StoreInterface->CachedPurchaseStateObject->ReadState = EOnlineAsyncTaskState::Done;
+        }
+
+        StoreInterface->TriggerOnInAppPurchaseCompleteDelegates(CompletionState);
+ 		// Ensure these flags are false to allow a subsequent transaction attempt to proceed
+		StoreInterface->bIsProductRequestInFlight = false;
+		StoreInterface->bIsPurchasing = false;
+		StoreInterface->bIsRestoringPurchases = false;
+	 
+        return true;
+     }];
+
+    [[SKPaymentQueue defaultQueue] finishTransaction: transaction];
+}
+
+
+- (void) requestProductData: (NSMutableSet*) productIDs
+{
+	UE_LOG(LogOnline, Verbose, TEXT( "FStoreKitHelper::requestProductData" ));
+
+	Request = [[SKProductsRequest alloc] initWithProductIdentifiers: productIDs];
+	Request.delegate = self;
+
+	[Request start];
+}
+
+
+- (void) makePurchase: (NSMutableSet*) productIDs
+{
+	UE_LOG(LogOnline, Verbose, TEXT( "FStoreKitHelper::makePurchase" ));
+
+	Request = [[SKProductsRequest alloc] initWithProductIdentifiers: productIDs];
+	Request.delegate = self;
+
+	[Request start];
+}
+
+
+- (void)productsRequest: (SKProductsRequest *)request didReceiveResponse:(SKProductsResponse *)response
+{
+	UE_LOG(LogOnline, Verbose, TEXT( "FStoreKitHelper::didReceiveResponse" ));
+	// Direct the response back to the store interface
+	[FIOSAsyncTask CreateTaskWithBlock : ^ bool(void)
+	{
+		IOnlineSubsystem* OnlineSub = IOnlineSubsystem::Get(IOS_SUBSYSTEM);
+		FOnlineStoreInterfaceIOS* StoreInterface = (FOnlineStoreInterfaceIOS*)OnlineSub->GetStoreInterface().Get();
+		StoreInterface->ProcessProductsResponse(response);
+
+		return true;
+	}];
+
+	[request autorelease];
+}
+
+-(void)requestDidFinish:(SKRequest*)request
+{
+#ifdef __IPHONE_7_0
+	if ([Request isKindOfClass : [SKReceiptRefreshRequest class]])
+	{
+		[[SKPaymentQueue defaultQueue] restoreCompletedTransactions];
+	}
+	
+	else if ([Request isKindOfClass : [SKProductsRequest class]])
+	{
+		// This is a placeholder to check if this case needs to be handled in the future
+	}
+	
+#endif
+}
+
+-(void)request:(SKRequest*)request didFailWithError: (NSError *)error
+{
+#ifdef __IPHONE_7_0
+	// Check if we were refreshing a receipt
+	if ([Request isKindOfClass : [SKReceiptRefreshRequest class]])
+	{
+		UE_LOG(LogOnline, Warning, TEXT( "FStoreKitHelper: SKReceiptRefreshRequest failed" ));
+		[self paymentQueue:[SKPaymentQueue defaultQueue] restoreCompletedTransactionsFailedWithError: error];
+		[Request release];
+	}
+
+	// Check if we are doing a product info request
+	else if ([Request isKindOfClass : [SKProductsRequest class]])
+	{
+		UE_LOG(LogOnline, Warning, TEXT( "FStoreKitHelper: SKProductsRequest failed" ));
+		[self paymentQueue:[SKPaymentQueue defaultQueue] requestProductDataFailedWithError: error];
+		[Request release];
+	}
+	else
+	{
+		UE_LOG(LogOnline, Warning, TEXT( "FStoreKitHelper: didFailWithError went unhandled" ));
+	}
+	
+#endif
+}
+
+-(void)restorePurchases
+{
+#ifdef __IPHONE_7_0
+	if ([IOSAppDelegate GetDelegate].OSVersion >= 7.0)
+	{
+		Request = [[SKReceiptRefreshRequest alloc] init];
+		Request.delegate = self;
+		[Request start];
+    }
+	else
+#endif
+	{
+#if __IPHONE_OS_VERSION_MIN_REQUIRED < __IPHONE_7_0
+		[[SKPaymentQueue defaultQueue] restoreCompletedTransactions];
+#endif
+	}
+}
+@end
+
+
+
+////////////////////////////////////////////////////////////////////
+/// FOnlineStoreInterfaceIOS implementation
+
+
+FOnlineStoreInterfaceIOS::FOnlineStoreInterfaceIOS() 
+{
+	UE_LOG(LogOnline, Verbose, TEXT( "FOnlineStoreInterfaceIOS::FOnlineStoreInterfaceIOS" ));
+	StoreHelper = [[FStoreKitHelper alloc] init];
+    
+    [[SKPaymentQueue defaultQueue] addTransactionObserver:StoreHelper];
+
+	bIsPurchasing = false;
+	bIsProductRequestInFlight = false;
+	bIsRestoringPurchases = false;
+}
+
+
+FOnlineStoreInterfaceIOS::~FOnlineStoreInterfaceIOS() 
+{
+	UE_LOG(LogOnline, Verbose, TEXT( "FOnlineStoreInterfaceIOS::~FOnlineStoreInterfaceIOS" ));
+	[StoreHelper release];
+}
+
+
+bool FOnlineStoreInterfaceIOS::QueryForAvailablePurchases(const TArray<FString>& ProductIDs, FOnlineProductInformationReadRef& InReadObject)
+{
+	UE_LOG(LogOnline, Verbose, TEXT( "FOnlineStoreInterfaceIOS::QueryForAvailablePurchases" ));
+	bool bSentAQueryRequest = false;
+
+	CachedReadObject = InReadObject;
+	
+	if( bIsPurchasing || bIsProductRequestInFlight || bIsRestoringPurchases )
+	{
+		UE_LOG(LogOnline, Verbose, TEXT( "FOnlineStoreInterfaceIOS::BeginPurchase - cannot make a purchase whilst one is in transaction." ));
+	}
+	else if (ProductIDs.Num() == 0)
+	{
+		UE_LOG(LogOnline, Verbose, TEXT( "There are no product IDs configured for Microtransactions in the engine.ini" ));
+	}
+	else
+	{
+		// autoreleased NSSet to hold IDs
+		NSMutableSet* ProductSet = [NSMutableSet setWithCapacity:ProductIDs.Num()];
+		for (int32 ProductIdx = 0; ProductIdx < ProductIDs.Num(); ProductIdx++)
+		{
+			NSString* ID = [NSString stringWithFString:ProductIDs[ ProductIdx ]];
+			// convert to NSString for the set objects
+			[ProductSet addObject:ID];
+		}
+        
+        dispatch_async(dispatch_get_main_queue(), ^
+        {
+            [StoreHelper requestProductData:ProductSet];
+        });
+
+		bIsProductRequestInFlight = true;
+
+		bSentAQueryRequest = true;
+	}
+
+	return bSentAQueryRequest;
+}
+
+
+bool FOnlineStoreInterfaceIOS::IsAllowedToMakePurchases()
+{
+	UE_LOG(LogOnline, Verbose, TEXT( "FOnlineStoreInterfaceIOS::IsAllowedToMakePurchases" ));
+	bool bCanMakePurchases = [SKPaymentQueue canMakePayments];
+	return bCanMakePurchases;
+}
+
+
+bool FOnlineStoreInterfaceIOS::BeginPurchase(const FInAppPurchaseProductRequest& ProductRequest, FOnlineInAppPurchaseTransactionRef& InPurchaseStateObject)
+{
+	UE_LOG(LogOnline, Verbose, TEXT( "FOnlineStoreInterfaceIOS::BeginPurchase" ));
+	bool bCreatedNewTransaction = false;
+	
+	const FString& ProductId = ProductRequest.ProductIdentifier;
+
+	if( bIsPurchasing || bIsProductRequestInFlight || bIsRestoringPurchases)
+	{
+		UE_LOG(LogOnline, Verbose, TEXT( "FOnlineStoreInterfaceIOS::BeginPurchase - cannot make a purchase whilst one is in transaction." ));
+        
+		TriggerOnInAppPurchaseCompleteDelegates(EInAppPurchaseState::Failed);
+	}
+	else if( IsAllowedToMakePurchases() )
+	{
+		UE_LOG(LogOnline, Verbose, TEXT("FOnlineStoreInterfaceIOS - Making a transaction of product %s"), *ProductId);
+
+		NSString* ID = [NSString stringWithFString : ProductId];
+        
+        dispatch_async(dispatch_get_main_queue(), ^
+        {
+            NSMutableSet* ProductSet = [NSMutableSet setWithCapacity:1];
+            [ProductSet addObject:ID];
+            
+            // Purchase the product through the StoreKit framework
+            [StoreHelper makePurchase:ProductSet];
+        });
+		
+		// Flag that we are purchasing so we can manage subsequent callbacks and transaction requests
+		bIsPurchasing = true;
+		
+		// We have created a new transaction on this pass.
+		bCreatedNewTransaction = true;
+
+		// Cache the readobject so we can add product information to it after the purchase
+		CachedPurchaseStateObject = InPurchaseStateObject;
+		CachedPurchaseStateObject->ReadState = EOnlineAsyncTaskState::InProgress;
+	}
+	else
+	{
+		UE_LOG(LogOnline, Verbose, TEXT("This device is not able to make purchases."));
+		InPurchaseStateObject->ReadState = EOnlineAsyncTaskState::Failed;
+        
+        TriggerOnInAppPurchaseCompleteDelegates(EInAppPurchaseState::NotAllowed);
+	}
+
+	return bCreatedNewTransaction;
+}
+
+
+bool FOnlineStoreInterfaceIOS::RestorePurchases(const TArray<FInAppPurchaseProductRequest>& ConsumableProductFlags, FOnlineInAppPurchaseRestoreReadRef& InReadObject)
+{
+	bool bSentAQueryRequest = false;
+	CachedPurchaseRestoreObject = InReadObject;
+
+	if (bIsPurchasing || bIsProductRequestInFlight || bIsRestoringPurchases)
+	{
+		UE_LOG(LogOnline, Verbose, TEXT("FOnlineStoreInterfaceIOS::BeginPurchase - cannot make a purchase whilst one is in transaction."));
+		TriggerOnInAppPurchaseRestoreCompleteDelegates(EInAppPurchaseState::Failed);
+	}
+	else
+	{
+		bIsRestoringPurchases = true;
+		[StoreHelper restorePurchases];
+		bSentAQueryRequest = true;
+	}
+
+	return bSentAQueryRequest;
+}
+
+void FOnlineStoreInterfaceIOS::ProcessProductsResponse( SKProductsResponse* Response )
+{
+	if( bIsPurchasing )
+	{
+		bool bWasSuccessful = [Response.products count] > 0;
+		if( [Response.products count] == 1 )
+		{
+			SKProduct* Product = [Response.products objectAtIndex:0];
+
+			NSNumberFormatter *numberFormatter = [[NSNumberFormatter alloc] init];
+			[numberFormatter setFormatterBehavior:NSNumberFormatterBehavior10_4];
+			[numberFormatter setNumberStyle:NSNumberFormatterCurrencyStyle];
+			[numberFormatter setLocale:Product.priceLocale];
+
+			UE_LOG(LogOnline, Log, TEXT("Making a purchase: Product: %s, Price: %s"), *FString([Product localizedTitle]), *FString([numberFormatter stringFromNumber:Product.price]));
+
+			FInAppPurchaseProductInfo PurchaseProductInfo;
+			PurchaseProductInfo.Identifier = [Product productIdentifier];
+			PurchaseProductInfo.DisplayName = [Product localizedTitle];
+			PurchaseProductInfo.DisplayDescription = [Product localizedDescription];
+			PurchaseProductInfo.DisplayPrice = [numberFormatter stringFromNumber : Product.price];
+			PurchaseProductInfo.CurrencyCode = [Product.priceLocale objectForKey : NSLocaleCurrencyCode];
+			PurchaseProductInfo.CurrencySymbol = [Product.priceLocale objectForKey : NSLocaleCurrencySymbol];
+			PurchaseProductInfo.DecimalSeparator = [Product.priceLocale objectForKey : NSLocaleDecimalSeparator];
+			PurchaseProductInfo.GroupingSeparator = [Product.priceLocale objectForKey : NSLocaleGroupingSeparator];
+
+			[numberFormatter release];
+
+			CachedPurchaseStateObject->ProvidedProductInformation = PurchaseProductInfo;
+            
+            dispatch_async(dispatch_get_main_queue(), ^
+            {
+                // now that we have recently refreshed the info, we can purchase it
+                SKPayment* Payment = [SKPayment paymentWithProduct:Product];
+                [[SKPaymentQueue defaultQueue] addPayment:Payment];
+            });
+		}
+		else if( [Response.products count] > 1 )
+		{
+			UE_LOG(LogOnline, Warning, TEXT("Wrong number of products, [%d], in the response when trying to make a single purchase"), [Response.products count]);
+		}
+		else
+		{
+			for(NSString *invalidProduct in Response.invalidProductIdentifiers)
+			{
+				UE_LOG(LogOnline, Error, TEXT("Problem in iTunes connect configuration for product: %s"), *FString(invalidProduct));
+			}
+		}
+
+		bIsPurchasing = false;
+	}
+	else if( bIsProductRequestInFlight )
+	{
+		bool bWasSuccessful = [Response.products count] > 0;
+		if( [Response.products count] == 0 && [Response.invalidProductIdentifiers count] == 0 )
+		{
+			UE_LOG(LogOnline, Warning, TEXT("Wrong number of products [%d] in the response when trying to make a single purchase"), [Response.products count]);
+		}
+
+		if( CachedReadObject.IsValid() )
+		{
+			for (SKProduct* Product in Response.products)
+			{
+				NSNumberFormatter *numberFormatter = [[NSNumberFormatter alloc] init];
+				[numberFormatter setFormatterBehavior : NSNumberFormatterBehavior10_4];
+				[numberFormatter setNumberStyle : NSNumberFormatterCurrencyStyle];
+				[numberFormatter setLocale : Product.priceLocale];
+
+				FInAppPurchaseProductInfo NewProductInfo;
+				NewProductInfo.Identifier = [Product productIdentifier];
+				NewProductInfo.DisplayName = [Product localizedTitle];
+				NewProductInfo.DisplayDescription = [Product localizedDescription];
+				NewProductInfo.DisplayPrice = [numberFormatter stringFromNumber : Product.price];
+				NewProductInfo.CurrencyCode = [Product.priceLocale objectForKey : NSLocaleCurrencyCode];
+				NewProductInfo.CurrencySymbol = [Product.priceLocale objectForKey : NSLocaleCurrencySymbol];
+				NewProductInfo.DecimalSeparator = [Product.priceLocale objectForKey : NSLocaleDecimalSeparator];
+				NewProductInfo.GroupingSeparator = [Product.priceLocale objectForKey : NSLocaleGroupingSeparator];
+
+				UE_LOG(LogOnline, Log, TEXT("\nProduct Identifier: %s, Name: %s, Description: %s, Price: %s\n"),
+					*NewProductInfo.Identifier,
+					*NewProductInfo.DisplayName,
+					*NewProductInfo.DisplayDescription,
+					*NewProductInfo.DisplayPrice);
+
+				[numberFormatter release];
+
+				CachedReadObject->ProvidedProductInformation.Add( NewProductInfo );
+			}
+		}
+		else
+		{
+			UE_LOG(LogOnline, Log, TEXT("Read Object is invalid."));
+		}
+
+		
+		for( NSString *invalidProduct in Response.invalidProductIdentifiers )
+		{
+			UE_LOG(LogOnline, Warning, TEXT("Problem in iTunes connect configuration for product: %s"), *FString(invalidProduct));
+		}
+
+		TriggerOnQueryForAvailablePurchasesCompleteDelegates( bWasSuccessful );
+
+		bIsProductRequestInFlight = false;
+	}
+}
\ No newline at end of file
diff --git a/Engine/Plugins/Online/IOS/OnlineSubsystemIOS/Source/Private/OnlineStoreInterfaceIOS.h b/Engine/Plugins/Online/IOS/OnlineSubsystemIOS/Source/Private/OnlineStoreInterfaceIOS.h
new file mode 100644
--- /dev/null
+++ b/Engine/Plugins/Online/IOS/OnlineSubsystemIOS/Source/Private/OnlineStoreInterfaceIOS.h
@@ -0,0 +1,106 @@
+// Copyright 1998-2016 Epic Games, Inc. All Rights Reserved.
+
+#pragma once
+
+#include "OnlineStoreInterface.h"
+
+// SK includes
+#include <StoreKit/SKRequest.h>
+#include <StoreKit/SKError.h>
+#include <StoreKit/SKProduct.h>
+#include <StoreKit/SKProductsRequest.h>
+#include <StoreKit/SKReceiptRefreshRequest.h>
+#include <StoreKit/SKPaymentTransaction.h>
+#include <StoreKit/SKPayment.h>
+#include <StoreKit/SKPaymentQueue.h>
+
+
+/**
+ * The resulting state of an iap transaction
+ */
+namespace EInAppPurchaseResult
+{
+    enum Type
+    {
+        Succeeded = 0,
+        RestoredFromServer,
+        Failed,
+        Cancelled,
+    };
+}
+
+
+/** Helper class, which allows us to manage IAP product information requests, AND transactions */
+@interface FStoreKitHelper : NSObject<SKProductsRequestDelegate, SKPaymentTransactionObserver, SKRequestDelegate>
+{
+};
+/** Store kit request object, holds information about the products we are purchasing, or querying. */
+@property (nonatomic, strong) SKRequest *Request;
+/** collection of available products attaced through a store kit request */
+@property (nonatomic, strong) NSArray *AvailableProducts;
+
+/** Helper fn to start a store kit purchase request */
+-(void)makePurchase:(NSMutableSet*)productIDs;
+/** Helper fn to start a store kit purchase information query request */
+- (void) requestProductData:(NSMutableSet*) productIDs;
+
+/** Helper fn to direct a product request response back to our store interface */
+- (void)productsRequest:(SKProductsRequest *)request didReceiveResponse:(SKProductsResponse *)response;
+
+/** Helper fn to restore previously purchased products */
+-(void)restorePurchases;
+@end
+
+
+/**
+ *	FOnlineStoreInterfaceIOS - Implementation of the online store for IOS
+ */
+class FOnlineStoreInterfaceIOS : public IOnlineStore
+{
+public:
+	/** C-tor */
+	FOnlineStoreInterfaceIOS();
+	/** Destructor */
+	virtual ~FOnlineStoreInterfaceIOS();
+
+	//~ Begin IOnlineStore Interface
+	virtual bool QueryForAvailablePurchases(const TArray<FString>& ProductIDs, FOnlineProductInformationReadRef& InReadObject) override;
+	virtual bool BeginPurchase(const FInAppPurchaseProductRequest& ProductRequest, FOnlineInAppPurchaseTransactionRef& InReadObject) override;
+	virtual bool IsAllowedToMakePurchases() override;
+	virtual bool RestorePurchases(const TArray<FInAppPurchaseProductRequest>& ConsumableProductFlags, FOnlineInAppPurchaseRestoreReadRef& InReadObject) override;
+	//~ End IOnlineStore Interface
+
+	/**
+	 * Process a response from the StoreKit
+	 */
+	void ProcessProductsResponse( SKProductsResponse* Response );
+    
+	/** Cached in-app purchase query object, used to provide the user with product information attained from the server */
+	FOnlineProductInformationReadPtr CachedReadObject;
+    
+	/** Cached in-app purchase transaction object, used to provide details to the user, of the product that has just been purchased. */
+	FOnlineInAppPurchaseTransactionPtr CachedPurchaseStateObject;
+
+	/** Cached in-app purchase restore transaction object, used to provide details to the developer about what products should be restored */
+	FOnlineInAppPurchaseRestoreReadPtr CachedPurchaseRestoreObject;
+
+
+	/** Flags which determine the state of our transaction buffer, only one action at a time */
+	bool bIsPurchasing, bIsProductRequestInFlight, bIsRestoringPurchases;
+
+private:
+	/** Access to the IOS Store kit interface */
+	FStoreKitHelper* StoreHelper;
+
+	/** Delegate fired when a query for purchases has completed, whether successful or unsuccessful */
+	FOnQueryForAvailablePurchasesComplete OnQueryForAvailablePurchasesCompleteDelegate;
+
+	/** Delegate fired when a purchase transaction has completed, whether successful or unsuccessful */
+	FOnInAppPurchaseComplete OnPurchaseCompleteDelegate;
+
+	/** Delegate fired when the purchase restoration has completed, whether successful or unsuccessful */
+	FOnInAppPurchaseRestoreComplete OnPurchaseRestoreCompleteDelegate;
+};
+
+typedef TSharedPtr<FOnlineStoreInterfaceIOS, ESPMode::ThreadSafe> FOnlineStoreInterfaceIOSPtr;
+
diff --git a/Engine/Plugins/Online/IOS/OnlineSubsystemIOS/Source/Private/OnlineSubsystemIOSPrivatePCH.h b/Engine/Plugins/Online/IOS/OnlineSubsystemIOS/Source/Private/OnlineSubsystemIOSPrivatePCH.h
new file mode 100644
--- /dev/null
+++ b/Engine/Plugins/Online/IOS/OnlineSubsystemIOS/Source/Private/OnlineSubsystemIOSPrivatePCH.h
@@ -0,0 +1,27 @@
+// Copyright 1998-2016 Epic Games, Inc. All Rights Reserved.
+
+#pragma once
+
+/** FName declaration of IOS subsystem */
+
+#include "Core.h"
+#include "CoreUObject.h"
+#include "SocketSubsystem.h"
+#include "ModuleManager.h"
+#include "IOSAppDelegate.h"
+
+class FOnlineSubsystemIOS;
+
+#include "OnlineSessionInterfaceIOS.h"
+#include "OnlineFriendsInterfaceIOS.h"
+#include "OnlineIdentityInterfaceIOS.h"
+#include "OnlineLeaderboardsInterfaceIOS.h"
+#include "OnlineStoreInterfaceIOS.h"
+#include "OnlineAchievementsInterfaceIOS.h"
+#include "OnlineExternalUIInterfaceIOS.h"
+#include "OnlineTurnBasedInterfaceIOS.h"
+#include "OnlineUserCloudInterfaceIOS.h"
+#include "OnlineSharedCloudInterfaceIOS.h"
+
+#include "OnlineSubsystemIOS.h"
+#include "OnlineSubsystemIOSModule.h"
\ No newline at end of file
diff --git a/Engine/Plugins/Online/IOS/OnlineSubsystemIOS/Source/Private/OnlineSubsystemIOSTypes.h b/Engine/Plugins/Online/IOS/OnlineSubsystemIOS/Source/Private/OnlineSubsystemIOSTypes.h
new file mode 100644
--- /dev/null
+++ b/Engine/Plugins/Online/IOS/OnlineSubsystemIOS/Source/Private/OnlineSubsystemIOSTypes.h
@@ -0,0 +1,102 @@
+// Copyright 1998-2016 Epic Games, Inc. All Rights Reserved.
+
+#pragma once
+
+#include "OnlineSubsystemTypes.h"
+#include "OnlineSubsystemIOSPackage.h"
+
+/**
+ * GameCenter specific implementation of the unique net id
+ */
+class FUniqueNetIdGameCenter :
+	public FUniqueNetId
+{
+PACKAGE_SCOPE:
+	/** Holds the net id for a player */
+	uint64 UniqueNetId;
+
+	/** Hidden on purpose */
+	FUniqueNetIdGameCenter() :
+		UniqueNetId(0)
+	{
+	}
+
+	/**
+	 * Copy Constructor
+	 *
+	 * @param Src the id to copy
+	 */
+	explicit FUniqueNetIdGameCenter(const FUniqueNetIdGameCenter& Src) :
+		UniqueNetId(Src.UniqueNetId)
+	{
+	}
+
+public:
+	/**
+	 * Constructs this object with the specified net id
+	 *
+	 * @param InUniqueNetId the id to set ours to
+	 */
+	explicit FUniqueNetIdGameCenter(uint64 InUniqueNetId) :
+		UniqueNetId(InUniqueNetId)
+	{
+	}
+
+	/**
+	 * Get the raw byte representation of this net id
+	 * This data is platform dependent and shouldn't be manipulated directly
+	 *
+	 * @return byte array of size GetSize()
+	 */
+	virtual const uint8* GetBytes() const override
+	{
+		return (uint8*)&UniqueNetId;
+	}
+
+	/**
+	 * Get the size of the id
+	 *
+	 * @return size in bytes of the id representation
+	 */
+	virtual int32 GetSize() const override
+	{
+		return sizeof(uint64);
+	}
+
+	/**
+	 * Check the validity of the id 
+	 *
+	 * @return true if this is a well formed ID, false otherwise
+	 */
+	virtual bool IsValid() const override
+	{
+		return UniqueNetId != 0;
+	}
+
+	/** 
+	 * Platform specific conversion to string representation of data
+	 *
+	 * @return data in string form 
+	 */
+	virtual FString ToString() const override
+	{
+		return FString::Printf(TEXT("%I64d"), UniqueNetId);
+	}
+
+	/**
+	 * Get a human readable representation of the net id
+	 * Shouldn't be used for anything other than logging/debugging
+	 *
+	 * @return id in string form 
+	 */
+	virtual FString ToDebugString() const override
+	{
+		return FString::Printf(TEXT("0%I64X"), UniqueNetId);
+	}
+
+	/** Needed for TMap::GetTypeHash() */
+	friend uint32 GetTypeHash(const FUniqueNetIdGameCenter& A)
+	{
+		return (uint32)(A.UniqueNetId) + ((uint32)((A.UniqueNetId) >> 32 ) * 23);
+	}
+};
diff --git a/Engine/Plugins/Online/OnlineFramework/Source/Hotfix/Hotfix.Build.cs b/Engine/Plugins/Online/OnlineFramework/Source/Hotfix/Hotfix.Build.cs
new file mode 100644
--- /dev/null
+++ b/Engine/Plugins/Online/OnlineFramework/Source/Hotfix/Hotfix.Build.cs
@@ -0,0 +1,19 @@
+// Copyright 1998-2016 Epic Games, Inc. All Rights Reserved.
+
+using UnrealBuildTool;
+
+public class Hotfix : ModuleRules
+{
+	public Hotfix(TargetInfo Target)
+    {
+        PrivateDependencyModuleNames.AddRange(
+			new string[] { 
+				"Core",
+				"CoreUObject",
+				"Engine",
+				"OnlineSubsystem",
+				"OnlineSubsystemUtils",
+			}
+			);
+    }
+}
diff --git a/Engine/Plugins/Online/OnlineFramework/Source/Hotfix/Private/HotfixModule.cpp b/Engine/Plugins/Online/OnlineFramework/Source/Hotfix/Private/HotfixModule.cpp
new file mode 100644
--- /dev/null
+++ b/Engine/Plugins/Online/OnlineFramework/Source/Hotfix/Private/HotfixModule.cpp
@@ -0,0 +1,14 @@
+// Copyright 1998-2016 Epic Games, Inc. All Rights Reserved.
+
+#include "HotfixPrivatePCH.h"
+#include "HotfixModule.h"
+
+IMPLEMENT_MODULE(FHotfixModule, Hotfix);
+
+void FHotfixModule::StartupModule()
+{
+}
+
+void FHotfixModule::ShutdownModule()
+{
+}
diff --git a/Engine/Plugins/Online/OnlineFramework/Source/Hotfix/Private/HotfixPrivatePCH.h b/Engine/Plugins/Online/OnlineFramework/Source/Hotfix/Private/HotfixPrivatePCH.h
new file mode 100644
--- /dev/null
+++ b/Engine/Plugins/Online/OnlineFramework/Source/Hotfix/Private/HotfixPrivatePCH.h
@@ -0,0 +1,8 @@
+// Copyright 1998-2016 Epic Games, Inc. All Rights Reserved.
+
+#pragma once
+
+#include "Core.h"
+#include "CoreUObject.h"
+#include "Engine.h"
+#include "Online.h"
