Copying //UE4/GamePlugins-Staging to //UE4/Dev-Main (Source: //GamePlugins/Main @ 3232656)

diff --git a/Engine/Plugins/Online/IOS/OnlineSubsystemIOS/Source/Private/OnlineStoreKitHelper.cpp b/Engine/Plugins/Online/IOS/OnlineSubsystemIOS/Source/Private/OnlineStoreKitHelper.cpp
new file mode 100644
--- /dev/null
+++ b/Engine/Plugins/Online/IOS/OnlineSubsystemIOS/Source/Private/OnlineStoreKitHelper.cpp
@@ -0,0 +1,703 @@
+// Copyright 1998-2016 Epic Games, Inc. All Rights Reserved.
+
+#include "OnlineSubsystemIOSPrivatePCH.h"
+#import "OnlineStoreKitHelper.h"
+
+/**
+ * Convert an Apple SKPaymentTransaction receipt into a string
+ * 
+ * @param transaction valid transaction to convert to hex encoded string
+ *
+ * @return hex encoded string with opaque data representing a completed transaction
+ */
+FString convertReceiptToString(const SKPaymentTransaction* transaction)
+{
+	FString ReceiptData;
+	
+#ifdef __IPHONE_7_0
+	if ([IOSAppDelegate GetDelegate].OSVersion >= 7.0)
+	{
+		NSURL* nsReceiptUrl = [[NSBundle mainBundle] appStoreReceiptURL];
+		NSData* nsReceiptData = [NSData dataWithContentsOfURL : nsReceiptUrl];
+		if (nsReceiptData)
+		{
+			NSString* nsEncodedReceiptData = [nsReceiptData base64EncodedStringWithOptions : NSDataBase64EncodingEndLineWithLineFeed];
+			ReceiptData = nsEncodedReceiptData;
+		}
+		else
+		{
+			UE_LOG(LogOnline, Log, TEXT("No receipt data found for transaction"));
+		}
+	}
+	else
+#endif
+	{
+#if __IPHONE_OS_VERSION_MIN_REQUIRED < __IPHONE_7_0
+		// If earlier than IOS 7, we will need to use the transactionReceipt
+		NSString* nsEncodedReceiptData = [transaction.transactionReceipt base64EncodedStringWithOptions : NSDataBase64EncodingEndLineWithLineFeed];
+		ReceiptData = nsEncodedReceiptData;
+#endif
+	}
+	
+	UE_LOG(LogOnline, VeryVerbose, TEXT("FStoreKitHelper::convertReceiptToString %s"), *ReceiptData);
+	return ReceiptData;
+}
+
+/**
+ * Retrieve the original transaction id from an Apple transaction object
+ * Ignores code comment in SKPaymentTransaction.h that it can only be found in "restored" state
+ * Successful attempts to repurchase already owned items (NOT restore purchase), will end in "purchased" state with an original transaction id
+ *
+ * @param Transaction the transaction to retrieve an original transaction id
+ *
+ * @return original transaction id for the transaction
+ */
+FString GetOriginalTransactionId(const SKPaymentTransaction* Transaction)
+{
+	SKPaymentTransaction* OriginalTransaction = nullptr;
+	if (Transaction.originalTransaction)
+	{
+		if (Transaction.transactionState != SKPaymentTransactionStateRestored)
+		{
+			UE_LOG(LogOnline, Log, TEXT("Original transaction id in state %d"), static_cast<int32>(Transaction.transactionState));
+		}
+		
+		int32 RecurseCount = 0;
+		
+		if (Transaction != Transaction.originalTransaction)
+		{
+			OriginalTransaction = Transaction.originalTransaction;
+			while (OriginalTransaction.originalTransaction && (RecurseCount < 100))
+			{
+				++RecurseCount;
+				OriginalTransaction = OriginalTransaction.originalTransaction;
+			}
+		}
+		
+		if (RecurseCount > 0)
+		{
+			UE_LOG(LogOnline, Log, TEXT("Original transaction id recurse count %d"), RecurseCount);
+		}
+	}
+	
+	return OriginalTransaction ? OriginalTransaction.transactionIdentifier : Transaction.transactionIdentifier;
+}
+
+////////////////////////////////////////////////////////////////////
+/// FSKProductsRequestHelper implementation
+/// Add a delegate to a product information request 
+
+@implementation FSKProductsRequestHelper
+@end
+
+////////////////////////////////////////////////////////////////////
+/// FStoreKitTransactionData implementation
+FStoreKitTransactionData::FStoreKitTransactionData(const SKPaymentTransaction* Transaction)
+	: ReceiptData(convertReceiptToString(Transaction))
+	, ErrorStr([Transaction.error localizedDescription])
+	, TransactionIdentifier(Transaction.transactionIdentifier)
+{
+	SKPayment* Payment = Transaction.payment;
+	if (Payment)
+	{
+		OfferId = Payment.productIdentifier;
+	}
+	
+	OriginalTransactionIdentifier = GetOriginalTransactionId(Transaction);
+}
+
+////////////////////////////////////////////////////////////////////
+/// FStoreKitHelper implementation
+
+@implementation FStoreKitHelper
+@synthesize Request;
+@synthesize AvailableProducts;
+
+- (id)init
+{
+	self = [super init];
+	return self;
+}
+
+-(void)paymentQueue:(SKPaymentQueue *)queue updatedTransactions : (NSArray *)transactions
+{
+	// Parse the generic transaction update into appropriate execution paths
+	UE_LOG(LogOnline, Log, TEXT("FStoreKitHelper::updatedTransactions"));
+	for (SKPaymentTransaction *transaction in transactions)
+	{
+		switch ([transaction transactionState])
+		{
+			case SKPaymentTransactionStatePurchased:
+				UE_LOG(LogOnline, Verbose, TEXT("FStoreKitHelper::completeTransaction"));
+				[self completeTransaction : transaction];
+				break;
+			case SKPaymentTransactionStateFailed:
+				UE_LOG(LogOnline, Verbose, TEXT("FStoreKitHelper::failedTransaction"));
+				[self failedTransaction : transaction];
+				break;
+			case SKPaymentTransactionStateRestored:
+				UE_LOG(LogOnline, Verbose, TEXT("FStoreKitHelper::restoreTransaction"));
+				[self restoreTransaction : transaction];
+				break;
+			case SKPaymentTransactionStatePurchasing:
+				UE_LOG(LogOnline, Verbose, TEXT("FStoreKitHelper::purchasingInProgress"));
+				[self purchaseInProgress : transaction];
+				continue;
+			case SKPaymentTransactionStateDeferred:
+				UE_LOG(LogOnline, Verbose, TEXT("FStoreKitHelper::purchaseDeferred"));
+				[self purchaseDeferred : transaction];
+				continue;
+			default:
+				UE_LOG(LogOnline, Verbose, TEXT("FStoreKitHelper::other: %d"), [transaction transactionState]);
+				break;
+		}
+	}
+}
+
+-(void)paymentQueue:(SKPaymentQueue *)queue removedTransactions : (NSArray *)transactions
+{
+	UE_LOG(LogOnline, Log, TEXT("FStoreKitHelper::removedTransactions"));
+}
+
+-(void)paymentQueueRestoreCompletedTransactionsFinished:(SKPaymentQueue *)queue
+{
+	UE_LOG(LogOnline, Log, TEXT("FStoreKitHelper::paymentQueueRestoreCompletedTransactionsFinished"));
+
+	[FIOSAsyncTask CreateTaskWithBlock : ^ bool(void)
+	{
+		IOnlineSubsystem* OnlineSub = IOnlineSubsystem::Get(IOS_SUBSYSTEM);
+
+		FOnlineStoreInterfaceIOS* StoreInterface = (FOnlineStoreInterfaceIOS*)OnlineSub->GetStoreInterface().Get();
+		if (StoreInterface->CachedPurchaseRestoreObject.IsValid())
+		{
+			StoreInterface->CachedPurchaseRestoreObject->ReadState = EOnlineAsyncTaskState::Done;
+		}
+		StoreInterface->TriggerOnInAppPurchaseRestoreCompleteDelegates(EInAppPurchaseState::Restored);
+		
+		return true;
+	}];
+}
+
+-(void)paymentQueue: (SKPaymentQueue *)queue restoreCompletedTransactionsFailedWithError : (NSError *)error
+{
+	UE_LOG(LogOnline, Log, TEXT("FStoreKitHelper::failedRestore - %s"), *FString([error localizedDescription]));
+	
+	EInAppPurchaseState::Type CompletionState = EInAppPurchaseState::Unknown;
+	switch (error.code)
+	{
+		case SKErrorPaymentCancelled:
+			CompletionState = EInAppPurchaseState::Cancelled;
+			break;
+		case SKErrorClientInvalid:
+		case SKErrorStoreProductNotAvailable:
+		case SKErrorPaymentInvalid:
+			CompletionState = EInAppPurchaseState::Invalid;
+			break;
+		case SKErrorPaymentNotAllowed:
+			CompletionState = EInAppPurchaseState::NotAllowed;
+			break;
+	}
+	
+	[FIOSAsyncTask CreateTaskWithBlock : ^ bool(void)
+	{
+		IOnlineSubsystem* OnlineSub = IOnlineSubsystem::Get(IOS_SUBSYSTEM);
+		FOnlineStoreInterfaceIOS* StoreInterface = (FOnlineStoreInterfaceIOS*)OnlineSub->GetStoreInterface().Get();
+		if (StoreInterface->CachedPurchaseRestoreObject.IsValid())
+		{
+			StoreInterface->CachedPurchaseRestoreObject->ReadState = EOnlineAsyncTaskState::Done;
+		}
+ 
+		StoreInterface->TriggerOnInAppPurchaseRestoreCompleteDelegates(CompletionState);
+	 
+		return true;
+	 }];
+}
+
+-(void)completeTransaction: (SKPaymentTransaction *)transaction
+{
+	UE_LOG(LogOnline, Log, TEXT("FStoreKitHelper::completeTransaction"));
+
+	[FIOSAsyncTask CreateTaskWithBlock : ^ bool(void)
+	{
+		IOnlineSubsystem* OnlineSub = IOnlineSubsystem::Get(IOS_SUBSYSTEM);
+		FOnlineStoreInterfaceIOS* StoreInterface = (FOnlineStoreInterfaceIOS*)OnlineSub->GetStoreInterface().Get();
+
+		if (StoreInterface->CachedPurchaseStateObject.IsValid())
+		{
+			 const FString ReceiptData = convertReceiptToString(transaction);
+			 
+			 StoreInterface->CachedPurchaseStateObject->ProvidedProductInformation.ReceiptData = ReceiptData;
+			 StoreInterface->CachedPurchaseStateObject->ProvidedProductInformation.TransactionIdentifier = transaction.transactionIdentifier;
+			 StoreInterface->CachedPurchaseStateObject->ReadState = EOnlineAsyncTaskState::Done;
+		}
+
+		StoreInterface->TriggerOnInAppPurchaseCompleteDelegates(EInAppPurchaseState::Success);
+
+		return true;
+	}];
+
+	// Remove the transaction from the payment queue.
+	[[SKPaymentQueue defaultQueue] finishTransaction:transaction];
+}
+
+-(void)restoreTransaction: (SKPaymentTransaction *)transaction
+{
+	UE_LOG(LogOnline, Log, TEXT("FStoreKitHelper::restoreTransaction"));
+
+	[FIOSAsyncTask CreateTaskWithBlock : ^ bool(void)
+	{
+		IOnlineSubsystem* OnlineSub = IOnlineSubsystem::Get(IOS_SUBSYSTEM);
+		FOnlineStoreInterfaceIOS* StoreInterface = (FOnlineStoreInterfaceIOS*)OnlineSub->GetStoreInterface().Get();
+
+		if (StoreInterface->CachedPurchaseRestoreObject.IsValid())
+		{
+			const FString ReceiptData = convertReceiptToString(transaction);
+	 
+			FInAppPurchaseRestoreInfo RestoreInfo;
+			RestoreInfo.Identifier = FString(transaction.originalTransaction.payment.productIdentifier);
+			RestoreInfo.ReceiptData = ReceiptData;
+			StoreInterface->CachedPurchaseRestoreObject->ProvidedRestoreInformation.Add(RestoreInfo);
+		}
+		
+		return true;
+	}];
+
+	[[SKPaymentQueue defaultQueue] finishTransaction:transaction];
+}
+
+-(void)failedTransaction: (SKPaymentTransaction *)transaction
+{
+	UE_LOG(LogOnline, Log, TEXT("FStoreKitHelper::failedTransaction - %s"), *FString([transaction.error localizedDescription]));
+
+	EInAppPurchaseState::Type CompletionState = EInAppPurchaseState::Unknown;
+	switch (transaction.error.code)
+	{
+		case SKErrorPaymentCancelled:
+			CompletionState = EInAppPurchaseState::Cancelled;
+			break;
+		case SKErrorClientInvalid:
+		case SKErrorStoreProductNotAvailable:
+		case SKErrorPaymentInvalid:
+			CompletionState = EInAppPurchaseState::Invalid;
+			break;
+		case SKErrorPaymentNotAllowed:
+			CompletionState = EInAppPurchaseState::NotAllowed;
+			break;
+	}
+
+	[FIOSAsyncTask CreateTaskWithBlock : ^ bool(void)
+	{
+		IOnlineSubsystem* OnlineSub = IOnlineSubsystem::Get(IOS_SUBSYSTEM);
+		FOnlineStoreInterfaceIOS* StoreInterface = (FOnlineStoreInterfaceIOS*)OnlineSub->GetStoreInterface().Get();
+		if (StoreInterface->CachedPurchaseStateObject.IsValid())
+		{
+			StoreInterface->CachedPurchaseStateObject->ReadState = EOnlineAsyncTaskState::Done;
+		}
+
+		StoreInterface->TriggerOnInAppPurchaseCompleteDelegates(CompletionState);
+		return true;
+	}];
+
+	[[SKPaymentQueue defaultQueue] finishTransaction:transaction];
+}
+
+-(void)purchaseInProgress: (SKPaymentTransaction *)transaction
+{
+	UE_LOG(LogOnline, Log, TEXT("FStoreKitHelper::purchaseInProgress"));
+}
+
+-(void)purchaseDeferred: (SKPaymentTransaction *)transaction
+{
+	UE_LOG(LogOnline, Log, TEXT("FStoreKitHelper::purchaseDeferred"));
+}
+
+-(void)requestProductData: (NSMutableSet*)productIDs
+{
+	UE_LOG(LogOnline, Verbose, TEXT("FStoreKitHelper::requestProductData"));
+
+	Request = [[SKProductsRequest alloc] initWithProductIdentifiers:productIDs];
+	Request.delegate = self;
+
+	[Request start];
+}
+
+-(void)makePurchase: (NSMutableSet*)productIDs
+{
+	UE_LOG(LogOnline, Verbose, TEXT("FStoreKitHelper::makePurchase"));
+
+	Request = [[SKProductsRequest alloc] initWithProductIdentifiers:productIDs];
+	Request.delegate = self;
+	
+	[Request start];
+}
+
+-(void)productsRequest: (SKProductsRequest *)request didReceiveResponse : (SKProductsResponse *)response
+{
+	UE_LOG(LogOnline, Verbose, TEXT("FStoreKitHelper::didReceiveResponse"));
+	// Direct the response back to the store interface
+	[FIOSAsyncTask CreateTaskWithBlock : ^ bool(void)
+	{
+		IOnlineSubsystem* OnlineSub = IOnlineSubsystem::Get(IOS_SUBSYSTEM);
+		FOnlineStoreInterfaceIOS* StoreInterface = (FOnlineStoreInterfaceIOS*)OnlineSub->GetStoreInterface().Get();
+		StoreInterface->ProcessProductsResponse(response);
+		
+		return true;
+	}];
+
+	[request autorelease];
+}
+
+-(void)requestDidFinish:(SKRequest*)request
+{
+#ifdef __IPHONE_7_0
+	if ([Request isKindOfClass : [SKReceiptRefreshRequest class]])
+	{
+		[[SKPaymentQueue defaultQueue] restoreCompletedTransactions];
+	}
+#endif
+}
+
+-(void)request:(SKRequest*)request didFailWithError : (NSError *)error
+{
+#ifdef __IPHONE_7_0
+	if ([Request isKindOfClass : [SKReceiptRefreshRequest class]])
+	{
+		[self paymentQueue : [SKPaymentQueue defaultQueue]  restoreCompletedTransactionsFailedWithError : error];
+		[Request release];
+	}
+#endif
+}
+
+-(void)restorePurchases
+{
+#ifdef __IPHONE_7_0
+	if ([IOSAppDelegate GetDelegate].OSVersion >= 7.0)
+	{
+		Request = [[SKReceiptRefreshRequest alloc] init];
+		Request.delegate = self;
+		[Request start];
+	}
+	else
+#endif
+	{
+#if __IPHONE_OS_VERSION_MIN_REQUIRED < __IPHONE_7_0
+		[[SKPaymentQueue defaultQueue] restoreCompletedTransactions];
+#endif
+	}
+}
+
+@end
+
+////////////////////////////////////////////////////////////////////
+/// FStoreKitHelperV2 implementation
+
+@implementation FStoreKitHelperV2
+
+- (id)init
+{
+	self = [super init];
+	self.PendingTransactions = [NSMutableSet setWithCapacity:5];
+	return self;
+}
+
+-(void)paymentQueueRestoreCompletedTransactionsFinished:(SKPaymentQueue *)queue
+{
+	UE_LOG(LogOnline, Log, TEXT("FStoreKitHelperV2::paymentQueueRestoreCompletedTransactionsFinished"));
+	
+	[FIOSAsyncTask CreateTaskWithBlock : ^ bool(void)
+	{
+		[self OnRestoreTransactionsComplete].Broadcast(EPurchaseTransactionState::Restored);
+		return true;
+	}];
+}
+
+-(void)paymentQueue: (SKPaymentQueue *)queue restoreCompletedTransactionsFailedWithError : (NSError *)error
+{
+	UE_LOG(LogOnline, Log, TEXT("FStoreKitHelperV2::failedRestore - %s"), *FString([error localizedDescription]));
+	
+	EPurchaseTransactionState CompletionState = EPurchaseTransactionState::Failed;
+	switch (error.code)
+	{
+		case SKErrorPaymentCancelled:
+			CompletionState = EPurchaseTransactionState::Canceled;
+			break;
+		case SKErrorClientInvalid:
+		case SKErrorStoreProductNotAvailable:
+		case SKErrorPaymentInvalid:
+			CompletionState = EPurchaseTransactionState::Invalid;
+			break;
+		case SKErrorPaymentNotAllowed:
+			CompletionState = EPurchaseTransactionState::NotAllowed;
+			break;
+	}
+	
+	[FIOSAsyncTask CreateTaskWithBlock : ^ bool(void)
+	{
+		self.OnRestoreTransactionsComplete.Broadcast(CompletionState);
+		return true;
+	}];
+}
+
+-(void)makePurchase:(NSArray*)products WithUserId: (const FString&) userId SimulateAskToBuy: (bool) bAskToBuy;
+{
+	UE_LOG(LogOnline, Verbose, TEXT("FStoreKitHelperV2::makePurchase by SKProduct with UserId"));
+	
+	for (SKProduct* Product in products)
+	{
+		SKMutablePayment* Payment = [SKMutablePayment paymentWithProduct:Product];
+		Payment.quantity = 1;
+		if ([IOSAppDelegate GetDelegate].OSVersion >= 8.3)
+		{
+			Payment.simulatesAskToBuyInSandbox = bAskToBuy;
+		}
+		if (!userId.IsEmpty())
+		{
+			// hash of username to detect irregular activity
+			Payment.applicationUsername = [NSString stringWithFString: userId];
+		}
+		[[SKPaymentQueue defaultQueue] addPayment:Payment];
+	}
+}
+
+-(void)makePurchase:(NSArray*)products SimulateAskToBuy: (bool) bAskToBuy
+{
+	UE_LOG(LogOnline, Verbose, TEXT("FStoreKitHelperV2::makePurchase by SKProduct"));
+	FString EmptyString;
+	[self makePurchase: products WithUserId: EmptyString SimulateAskToBuy: bAskToBuy];
+}
+
+-(void)requestProductData: (NSMutableSet*)productIDs WithDelegate : (const FOnQueryOnlineStoreOffersComplete&)delegate
+{
+	UE_LOG(LogOnline, Verbose, TEXT("FStoreKitHelperV2::requestProductData"));
+	
+	FSKProductsRequestHelper* TempRequest = [[FSKProductsRequestHelper alloc] initWithProductIdentifiers:productIDs];
+	TempRequest.OfferDelegate = delegate;
+	TempRequest.delegate = self;
+	self.Request = TempRequest;
+	
+	[self.Request start];
+}
+
+-(void)productsRequest: (SKProductsRequest *)request didReceiveResponse : (SKProductsResponse *)response
+{
+	if (request == self.Request)
+	{
+		UE_LOG(LogOnline, Verbose, TEXT("FStoreKitHelperV2::didReceiveResponse"));
+		
+		if ([request isKindOfClass : [FSKProductsRequestHelper class]])
+		{
+			FSKProductsRequestHelper* Helper = (FSKProductsRequestHelper*)request;
+
+			[FIOSAsyncTask CreateTaskWithBlock : ^ bool(void)
+			{
+				// Notify listeners of the request completion
+				self.OnProductRequestResponse.Broadcast(response, Helper.OfferDelegate);
+				return true;
+			}];
+		}
+		else
+		{
+			UE_LOG(LogOnline, Warning, TEXT("Wrong class associated with product request"));
+		}
+		
+		[request autorelease];
+	}
+}
+
+-(void)completeTransaction: (SKPaymentTransaction *)transaction
+{
+	FStoreKitTransactionData TransactionData(transaction);
+	UE_LOG(LogOnline, Log, TEXT("FStoreKitHelperV2::completeTransaction - %s"), *TransactionData.ToDebugString());
+	
+	EPurchaseTransactionState Result = EPurchaseTransactionState::Failed;
+	
+	SKPayment* Payment = transaction.payment;
+	if (Payment)
+	{
+		Result = EPurchaseTransactionState::Purchased;
+	}
+	
+	[FIOSAsyncTask CreateTaskWithBlock : ^ bool(void)
+	{
+	    // Notify listeners of the request completion
+		self.OnTransactionCompleteResponse.Broadcast(Result, TransactionData);
+		return true;
+	}];
+
+	// Transaction must be finalized before removed from the queue
+	[self.PendingTransactions addObject:transaction];
+}
+
+-(void)failedTransaction: (SKPaymentTransaction *)transaction
+{
+	FStoreKitTransactionData TransactionData(transaction);
+	UE_LOG(LogOnline, Log, TEXT("FStoreKitHelperV2::failedTransaction - %s"), *TransactionData.ToDebugString());
+
+	EPurchaseTransactionState CompletionState = EPurchaseTransactionState::Failed;
+	switch (transaction.error.code)
+	{
+		case SKErrorPaymentCancelled:
+			CompletionState = EPurchaseTransactionState::Canceled;
+			break;
+		case SKErrorClientInvalid:
+		case SKErrorStoreProductNotAvailable:
+		case SKErrorPaymentInvalid:
+			CompletionState = EPurchaseTransactionState::Invalid;
+			break;
+		case SKErrorPaymentNotAllowed:
+			CompletionState = EPurchaseTransactionState::NotAllowed;
+			break;
+	}
+	
+	[FIOSAsyncTask CreateTaskWithBlock : ^ bool(void)
+	{
+	    // Notify listeners of the request completion
+		self.OnTransactionCompleteResponse.Broadcast(CompletionState, TransactionData);
+		return true;
+	}];
+	
+	// Remove the transaction from the payment queue.
+	[[SKPaymentQueue defaultQueue] finishTransaction:transaction];
+}
+
+-(void)restoreTransaction: (SKPaymentTransaction *)transaction
+{
+	FStoreKitTransactionData TransactionData(transaction);
+    UE_LOG(LogOnline, Log, TEXT("FStoreKitHelperV2::restoreTransaction - %s"), *TransactionData.ToDebugString());
+
+	[FIOSAsyncTask CreateTaskWithBlock : ^ bool(void)
+	{
+		self.OnTransactionRestored.Broadcast(TransactionData);
+		return true;
+	}];
+	
+	// @todo Transaction must be finalized before removed from the queue?
+	//[self.PendingTransactions addObject:transaction];
+	// Remove the transaction from the payment queue.
+	[[SKPaymentQueue defaultQueue] finishTransaction:transaction];
+}
+
+-(void)purchaseInProgress: (SKPaymentTransaction *)transaction
+{
+	FStoreKitTransactionData TransactionData(transaction);
+	UE_LOG(LogOnline, Log, TEXT("FStoreKitHelperV2::purchaseInProgress - %s"), *TransactionData.ToDebugString());
+	
+	[FIOSAsyncTask CreateTaskWithBlock : ^ bool(void)
+	{
+		// Notify listeners a purchase is in progress
+		self.OnTransactionPurchaseInProgress.Broadcast(TransactionData);
+		return true;
+	}];
+}
+
+-(void)purchaseDeferred: (SKPaymentTransaction *)transaction
+{
+	FStoreKitTransactionData TransactionData(transaction);
+	UE_LOG(LogOnline, Log, TEXT("FStoreKitHelperV2::purchaseDeferred - %s"), *TransactionData.ToDebugString());
+	
+	[FIOSAsyncTask CreateTaskWithBlock : ^ bool(void)
+	{
+	    // Notify listeners a purchase has been deferred
+		self.OnTransactionDeferred.Broadcast(TransactionData);
+		return true;
+	}];
+}
+
+-(void)finalizeTransaction: (const FString&) receiptId
+{
+	UE_LOG(LogOnline, Verbose, TEXT("FStoreKitHelperV2::finalizeTransaction - %s"), *receiptId);
+	for (SKPaymentTransaction* pendingTransaction in self.PendingTransactions)
+	{
+		const FString transId = pendingTransaction.transactionIdentifier;
+		const FString originalTransId = GetOriginalTransactionId(pendingTransaction);
+		UE_LOG(LogOnline, Verbose, TEXT("FStoreKitHelperV2::checking - id: %s origId: %s"), *transId,  *originalTransId);
+		
+		if ((!originalTransId.IsEmpty()) && (originalTransId == receiptId))
+		{
+			UE_LOG(LogOnline, Log, TEXT("FStoreKitHelperV2::finalizeTransaction - %s"), *receiptId);
+
+			[self.PendingTransactions removeObject:pendingTransaction];
+			// Remove the transaction from the payment queue.
+			[[SKPaymentQueue defaultQueue] finishTransaction:pendingTransaction];
+			break;
+		}
+	}
+}
+
+-(void)dumpAppReceipt
+{
+#ifdef __IPHONE_7_0
+	FString receiptData = convertReceiptToString(nullptr);
+	UE_LOG(LogOnline, Verbose, TEXT("FStoreKitHelper::dumpAppReceipt"));
+	UE_LOG(LogOnline, Verbose, TEXT("%s"), *receiptData);
+#endif
+}
+
+-(FOnProductsRequestResponse&)OnProductRequestResponse
+{
+	return _OnProductRequestResponse;
+}
+
+-(FDelegateHandle)AddOnProductRequestResponse: (const FOnProductsRequestResponseDelegate&) Delegate
+{
+	_OnProductRequestResponse.Add(Delegate);
+	return Delegate.GetHandle();
+}
+
+-(FOnTransactionCompleteIOS&)OnTransactionCompleteResponse
+{
+	return _OnTransactionCompleteResponse;
+}
+
+-(FDelegateHandle)AddOnTransactionComplete: (const FOnTransactionCompleteIOSDelegate&) Delegate
+{
+	_OnTransactionCompleteResponse.Add(Delegate);
+	return Delegate.GetHandle();
+}
+
+-(FOnTransactionRestoredIOS&)OnTransactionRestored
+{
+	return _OnTransactionRestored;
+}
+
+-(FDelegateHandle)AddOnTransactionRestored: (const FOnTransactionRestoredIOSDelegate&) Delegate
+{
+	_OnTransactionRestored.Add(Delegate);
+	return Delegate.GetHandle();
+}
+
+-(FOnRestoreTransactionsCompleteIOS&)OnRestoreTransactionsComplete
+{
+	return _OnRestoreTransactionsComplete;
+}
+
+-(FDelegateHandle)AddOnRestoreTransactionsComplete: (const FOnRestoreTransactionsCompleteIOSDelegate&) Delegate
+{
+	_OnRestoreTransactionsComplete.Add(Delegate);
+	return Delegate.GetHandle();
+}
+
+-(FOnTransactionProgress&)OnTransactionPurchaseInProgress
+{
+	return _OnTransactionPurchaseInProgress;
+}
+
+-(FDelegateHandle)AddOnPurchaseInProgress: (const FOnTransactionProgressDelegate&) Delegate
+{
+	_OnTransactionPurchaseInProgress.Add(Delegate);
+	return Delegate.GetHandle();
+}
+
+-(FOnTransactionProgress&)OnTransactionDeferred
+{
+	return _OnTransactionDeferred;
+}
+
+-(FDelegateHandle)AddOnTransactionDeferred: (const FOnTransactionProgressDelegate&) Delegate
+{
+	_OnTransactionDeferred.Add(Delegate);
+	return Delegate.GetHandle();
+}
+
+@end
+
diff --git a/Engine/Plugins/Online/IOS/OnlineSubsystemIOS/Source/Private/OnlineStoreKitHelper.h b/Engine/Plugins/Online/IOS/OnlineSubsystemIOS/Source/Private/OnlineStoreKitHelper.h
new file mode 100644
--- /dev/null
+++ b/Engine/Plugins/Online/IOS/OnlineSubsystemIOS/Source/Private/OnlineStoreKitHelper.h
@@ -0,0 +1,208 @@
+// SK includes
+#import <StoreKit/SKRequest.h>
+#import <StoreKit/SKError.h>
+#import <StoreKit/SKProduct.h>
+#import <StoreKit/SKProductsRequest.h>
+#import <StoreKit/SKReceiptRefreshRequest.h>
+#import <StoreKit/SKPaymentTransaction.h>
+#import <StoreKit/SKPayment.h>
+#import <StoreKit/SKPaymentQueue.h>
+
+#include "OnlineSubsystemTypes.h"
+
+/**
+ * Holds in a common format the data that comes out of an SKPaymentTransaction
+ */
+struct FStoreKitTransactionData
+{
+	FStoreKitTransactionData(const SKPaymentTransaction* Transaction);
+	
+	/** @return a string that prints useful debug information about this transaction */
+	FString ToDebugString() const
+	{
+		return FString::Printf(TEXT("OfferId: %s TransactionId: %s%s ReceiptData: %s%s"),
+						*OfferId,
+						*TransactionIdentifier,
+						OriginalTransactionIdentifier.IsEmpty() ? TEXT("") : *FString::Printf(TEXT(" OriginalTransactionId: %s"), *OriginalTransactionIdentifier),
+						*ReceiptData,
+						ErrorStr.IsEmpty() ? TEXT("") : *FString::Printf(TEXT(" Error: %s"), *ErrorStr));
+	}
+	
+	/** @return offer id for this transaction */
+	const FString& GetOfferId() const { return OfferId; }
+	/** @return receipt data for this transaction */
+	const FString& GetReceiptData() const { return ReceiptData; }
+	/** @return error string for this transaction, if applicable */
+	const FString& GetErrorStr() const { return ErrorStr; }
+	
+	/** @return the correct transaction identifier relative to the original purchase */
+	const FString& GetTransactionIdentifier() const
+	{
+		if (!OriginalTransactionIdentifier.IsEmpty())
+		{
+			return OriginalTransactionIdentifier;
+		}
+		else
+		{
+			return TransactionIdentifier;
+		}
+	}
+	
+private:
+	
+	/** iTunesConnect offer id */
+	FString OfferId;
+	/** Opaque store receipt data */
+	FString ReceiptData;
+	/** Error on the transaction, if applicable */
+	FString ErrorStr;
+	/** Unique transaction id */
+	FString TransactionIdentifier;
+	/** Original transaction id if restored purchase */
+	FString OriginalTransactionIdentifier;
+};
+
+/**
+ * Delegate fires when a single transaction has completed (may be part of many offers in a single user purchase)
+ *
+ * @param Result result of the transaction
+ * @param TransactionData data associated with the completed transaction
+ */
+DECLARE_MULTICAST_DELEGATE_TwoParams(FOnTransactionCompleteIOS, EPurchaseTransactionState /*Result*/, const FStoreKitTransactionData& /*TransactionData*/);
+typedef FOnTransactionCompleteIOS::FDelegate FOnTransactionCompleteIOSDelegate;
+/**
+ * Delegate fires when a single restored transaction has completed (may be part of many user purchases that are restored)
+ *
+ * @param TransactionData data associated with the restored transaction
+ */
+DECLARE_MULTICAST_DELEGATE_OneParam(FOnTransactionRestoredIOS, const FStoreKitTransactionData& /*TransactionData*/);
+typedef FOnTransactionRestoredIOS::FDelegate FOnTransactionRestoredIOSDelegate;
+
+/**
+ * Delegate fires when a transaction progress event occurs
+ *
+ * @param TransactionData data associated with the event
+ */
+DECLARE_MULTICAST_DELEGATE_OneParam(FOnTransactionProgress, const FStoreKitTransactionData& /*TransactionData*/);
+typedef FOnTransactionProgress::FDelegate FOnTransactionProgressDelegate;
+
+/**
+ * Delegate fires when the entire restore transactions operation has completed
+ *
+ * @param Result aggregate result of the restore operation
+ */
+DECLARE_MULTICAST_DELEGATE_OneParam(FOnRestoreTransactionsCompleteIOS, EPurchaseTransactionState /*Result*/);
+typedef FOnRestoreTransactionsCompleteIOS::FDelegate FOnRestoreTransactionsCompleteIOSDelegate;
+
+/**
+ * Delegate fires when a product request for offer(s) has completed
+ *
+ * @param Response the array of product information for the products that were queried
+ * @param CompletionDelegate the delegate to fire externally when this operation completes
+ */
+DECLARE_MULTICAST_DELEGATE_TwoParams(FOnProductsRequestResponse, SKProductsResponse* /*Response*/, const FOnQueryOnlineStoreOffersComplete& /*CompletionDelegate*/);
+typedef FOnProductsRequestResponse::FDelegate FOnProductsRequestResponseDelegate;
+
+/**
+ * Product request helper to store a delegate for this individual product request
+ */
+@interface FSKProductsRequestHelper : SKProductsRequest
+{
+};
+
+/** Delegate to fire when this product request completes with the store kit */
+@property FOnQueryOnlineStoreOffersComplete OfferDelegate;
+
+@end
+
+/** 
+ * Helper class, which allows us to manage IAP product information requests, AND transactions
+ * (legacy version, used by OnlineStoreInterface.h, mutually exclusive with FStoreKitHelperV2)
+ */
+@interface FStoreKitHelper : NSObject<SKProductsRequestDelegate, SKPaymentTransactionObserver, SKRequestDelegate>
+{
+};
+
+/** Store kit request object, holds information about the products we are purchasing, or querying. */
+@property (nonatomic, strong) SKRequest *Request;
+/** collection of available products attaced through a store kit request */
+@property (nonatomic, strong) NSArray *AvailableProducts;
+
+/** Helper fn to start a store kit purchase request */
+-(void)makePurchase:(NSMutableSet*)productIDs;
+
+/** Helper fn to start a store kit purchase information query request */
+-(void)requestProductData:(NSMutableSet*)productIDs;
+
+/** Helper fn to direct a product request response back to our store interface */
+-(void)productsRequest:(SKProductsRequest *)request didReceiveResponse : (SKProductsResponse *)response;
+
+/** Helper fn to restore previously purchased products */
+-(void)restorePurchases;
+@end
+
+/**
+ * Helper class, which allows us to manage IAP product information requests, AND transactions
+ * (version used by OnlineStoreIOS.h and OnlinePurchaseIOS.h)
+ */
+@interface FStoreKitHelperV2 : FStoreKitHelper
+{
+	/** delegate fired when a product request completes */
+	FOnProductsRequestResponse _OnProductRequestResponse;
+	/** delegate fired when a single purchase transaction has completed (may be a part of multiple requests at once) */
+	FOnTransactionCompleteIOS _OnTransactionCompleteResponse;
+	/** delegate fired when a single transaction is restored (may be a part of many restored purchases) */
+	FOnTransactionRestoredIOS _OnTransactionRestored;
+	/** delegate fired when all transactions have been restored */
+	FOnRestoreTransactionsCompleteIOS _OnRestoreTransactionsComplete;
+	/** delegate fired when a purchase in progress is detected */
+	FOnTransactionProgress _OnTransactionPurchaseInProgress;
+	/** delegate fired when a deferred purchase is detected */
+	FOnTransactionProgress _OnTransactionDeferred;
+};
+
+/** list of known pending transactions that are in the queue but not marked as redeemed by the app */
+@property (nonatomic, strong) NSMutableSet* PendingTransactions;
+
+-(FDelegateHandle)AddOnProductRequestResponse: (const FOnProductsRequestResponseDelegate&) Delegate;
+-(FDelegateHandle)AddOnTransactionComplete: (const FOnTransactionCompleteIOSDelegate&) Delegate;
+-(FDelegateHandle)AddOnTransactionRestored: (const FOnTransactionRestoredIOSDelegate&) Delegate;
+-(FDelegateHandle)AddOnRestoreTransactionsComplete: (const FOnRestoreTransactionsCompleteIOSDelegate&) Delegate;
+-(FDelegateHandle)AddOnPurchaseInProgress: (const FOnTransactionProgressDelegate&) Delegate;
+-(FDelegateHandle)AddOnTransactionDeferred: (const FOnTransactionProgressDelegate&) Delegate;
+
+/**
+ * Remove a transaction from the queue once it has been properly credited to the user
+ *
+ * @param receiptId id of the no longer pending transaction
+ */
+-(void)finalizeTransaction: (const FString&) receiptId;
+
+/**
+ * Make a purchase with the app store
+ *
+ * @param products array of SKProducts to purchase
+ */
+-(void)makePurchase:(NSArray*)products SimulateAskToBuy: (bool) bAskToBuy;
+
+/**
+ * Make a purchase with the app store
+ *
+ * @param products array of SKProducts to purchase
+ * @param userId hashed and opaque representation of platform user id
+ */
+-(void)makePurchase:(NSArray*)products WithUserId: (const FString&) userId SimulateAskToBuy: (bool) bAskToBuy;
+
+/**
+ * Make a request for product information
+ *
+ * @param productIDs array of offer ids to get information for
+ * @param delegate delegate to fire when operation completes
+ */
+-(void)requestProductData:(NSMutableSet*)productIDs WithDelegate : (const FOnQueryOnlineStoreOffersComplete&)delegate;
+
+/**
+ * Dumps the base64 encoded app receipt to log
+ */
+-(void)dumpAppReceipt;
+@end
diff --git a/Engine/Plugins/Online/IOS/OnlineSubsystemIOS/Source/Private/OnlineSubsystemIOS.cpp b/Engine/Plugins/Online/IOS/OnlineSubsystemIOS/Source/Private/OnlineSubsystemIOS.cpp
--- a/Engine/Plugins/Online/IOS/OnlineSubsystemIOS/Source/Private/OnlineSubsystemIOS.cpp
+++ b/Engine/Plugins/Online/IOS/OnlineSubsystemIOS/Source/Private/OnlineSubsystemIOS.cpp
@@ -1,7 +1,7 @@
 // Copyright 1998-2017 Epic Games, Inc. All Rights Reserved.
 
 #include "OnlineSubsystemIOSPrivatePCH.h"
-
+#import "OnlineStoreKitHelper.h"
 
 FOnlineSubsystemIOS::FOnlineSubsystemIOS()
 {
@@ -78,6 +78,16 @@ IOnlineStorePtr FOnlineSubsystemIOS::GetStoreInterface() const
 	return StoreInterface;
 }
 
+IOnlineStoreV2Ptr FOnlineSubsystemIOS::GetStoreV2Interface() const
+{
+	return StoreV2Interface;
+}
+
+IOnlinePurchasePtr FOnlineSubsystemIOS::GetPurchaseInterface() const
+{
+	return PurchaseInterface;
+}
+
 IOnlineEventsPtr FOnlineSubsystemIOS::GetEventsInterface() const
 {
 	return nullptr;
@@ -146,15 +156,42 @@ bool FOnlineSubsystemIOS::Init()
         UserCloudInterface = MakeShareable(new FOnlineUserCloudInterfaceIOS());
         SharedCloudInterface = MakeShareable(new FOnlineSharedCloudInterfaceIOS());
 	}
-
-	if( IsInAppPurchasingEnabled() )
+	
+	if(IsInAppPurchasingEnabled())
 	{
-		StoreInterface = MakeShareable(new FOnlineStoreInterfaceIOS());
+		if (IsV2StoreEnabled())
+		{
+			StoreV2Interface = MakeShareable(new FOnlineStoreIOS(this));
+			PurchaseInterface = MakeShareable(new FOnlinePurchaseIOS(this));
+			InitStoreKitHelper();
+		}
+		else
+		{
+			StoreInterface = MakeShareable(new FOnlineStoreInterfaceIOS());
+		}
 	}
 
 	return bSuccessfullyStartedUp;
 }
 
+void FOnlineSubsystemIOS::InitStoreKitHelper()
+{
+	StoreHelper = [[FStoreKitHelperV2 alloc] init];
+	[StoreHelper retain];
+	
+	// Give each interface a chance to bind to the store kit helper
+	StoreV2Interface->InitStoreKit(StoreHelper);
+	PurchaseInterface->InitStoreKit(StoreHelper);
+	
+	// Bind the observer after the interfaces have bound their delegates
+	[[SKPaymentQueue defaultQueue] addTransactionObserver:StoreHelper];
+}
+
+void FOnlineSubsystemIOS::CleanupStoreKitHelper()
+{
+	[StoreHelper release];
+}
+
 bool FOnlineSubsystemIOS::Tick(float DeltaTime)
 {
 	if (!FOnlineSubsystemImpl::Tick(DeltaTime))
@@ -175,6 +212,30 @@ bool FOnlineSubsystemIOS::Shutdown()
 	UE_LOG(LogOnline, Display, TEXT("FOnlineSubsystemIOS::Shutdown()"));
 
 	bSuccessfullyShutdown = FOnlineSubsystemImpl::Shutdown();
+	
+#define DESTRUCT_INTERFACE(Interface) \
+	if (Interface.IsValid()) \
+	{ \
+		ensure(Interface.IsUnique()); \
+		Interface = nullptr; \
+	}
+	
+	DESTRUCT_INTERFACE(SessionInterface);
+	DESTRUCT_INTERFACE(IdentityInterface);
+	DESTRUCT_INTERFACE(FriendsInterface);
+	DESTRUCT_INTERFACE(LeaderboardsInterface);
+	DESTRUCT_INTERFACE(AchievementsInterface);
+	DESTRUCT_INTERFACE(ExternalUIInterface);
+	DESTRUCT_INTERFACE(TurnBasedInterface);
+	DESTRUCT_INTERFACE(UserCloudInterface);
+	DESTRUCT_INTERFACE(SharedCloudInterface);
+	DESTRUCT_INTERFACE(StoreInterface);
+	DESTRUCT_INTERFACE(StoreV2Interface);
+	DESTRUCT_INTERFACE(PurchaseInterface);
+	
+	// Cleanup after the interfaces are free
+	CleanupStoreKitHelper();
+	
 	return bSuccessfullyShutdown;
 }
 
@@ -201,9 +262,19 @@ bool FOnlineSubsystemIOS::IsEnabled()
 	return bEnableGameCenter || bEnableCloudKit || IsInAppPurchasingEnabled();
 }
 
+bool FOnlineSubsystemIOS::IsV2StoreEnabled()
+{
+	bool bUseStoreV2 = false;
+	GConfig->GetBool(TEXT("OnlineSubsystemIOS.Store"), TEXT("bUseStoreV2"), bUseStoreV2, GEngineIni);
+	return bUseStoreV2;
+}
+
 bool FOnlineSubsystemIOS::IsInAppPurchasingEnabled()
 {
 	bool bEnableIAP = false;
 	GConfig->GetBool(TEXT("OnlineSubsystemIOS.Store"), TEXT("bSupportsInAppPurchasing"), bEnableIAP, GEngineIni);
-	return bEnableIAP;
+	
+	bool bEnableIAP1 = false;
+	GConfig->GetBool(TEXT("OnlineSubsystemIOS.Store"), TEXT("bSupportInAppPurchasing"), bEnableIAP1, GEngineIni);
+	return bEnableIAP || bEnableIAP1;
 }
diff --git a/Engine/Plugins/Online/IOS/OnlineSubsystemIOS/Source/Private/OnlineSubsystemIOSPrivatePCH.h b/Engine/Plugins/Online/IOS/OnlineSubsystemIOS/Source/Private/OnlineSubsystemIOSPrivatePCH.h
--- a/Engine/Plugins/Online/IOS/OnlineSubsystemIOS/Source/Private/OnlineSubsystemIOSPrivatePCH.h
+++ b/Engine/Plugins/Online/IOS/OnlineSubsystemIOS/Source/Private/OnlineSubsystemIOSPrivatePCH.h
@@ -19,6 +19,8 @@ class FOnlineSubsystemIOS;
 #include "OnlineIdentityInterfaceIOS.h"
 #include "OnlineLeaderboardsInterfaceIOS.h"
 #include "OnlineStoreInterfaceIOS.h"
+#include "OnlineStoreIOS.h"
+#include "OnlinePurchaseIOS.h"
 #include "OnlineAchievementsInterfaceIOS.h"
 #include "OnlineExternalUIInterfaceIOS.h"
 #include "OnlineTurnBasedInterfaceIOS.h"
@@ -26,4 +28,5 @@ class FOnlineSubsystemIOS;
 #include "OnlineSharedCloudInterfaceIOS.h"
 
 #include "OnlineSubsystemIOS.h"
-#include "OnlineSubsystemIOSModule.h"
\ No newline at end of file
+#include "OnlineSubsystemIOSModule.h"
+
diff --git a/Engine/Plugins/Online/IOS/OnlineSubsystemIOS/Source/Public/OnlinePurchaseIOS.h b/Engine/Plugins/Online/IOS/OnlineSubsystemIOS/Source/Public/OnlinePurchaseIOS.h
new file mode 100644
--- /dev/null
+++ b/Engine/Plugins/Online/IOS/OnlineSubsystemIOS/Source/Public/OnlinePurchaseIOS.h
@@ -0,0 +1,153 @@
+// Copyright 1998-2016 Epic Games, Inc. All Rights Reserved.
+
+#pragma once
+
+#include "OnlinePurchaseInterface.h"
+
+struct FStoreKitTransactionData;
+@class FStoreKitHelperV2;
+
+/**
+ * Info used to cache and track orders in progress.
+ */
+class FOnlinePurchasePendingTransactionIOS
+{
+public:
+	FOnlinePurchasePendingTransactionIOS(
+		const FPurchaseCheckoutRequest& InCheckoutRequest,
+		const FUniqueNetId& InUserId,
+		const EPurchaseTransactionState InPendingTransactionState = EPurchaseTransactionState::NotStarted,
+		const FOnPurchaseCheckoutComplete& InCheckoutCompleteDelegate = FOnPurchaseCheckoutComplete()
+		)
+		: CheckoutRequest(InCheckoutRequest)
+		, UserId(InUserId)
+		, CheckoutCompleteDelegate(InCheckoutCompleteDelegate)
+	{
+		PendingPurchaseInfo.TransactionState = InPendingTransactionState;
+		
+		// Setup purchase state for all pending offers
+		OfferPurchaseStates.AddZeroed(CheckoutRequest.PurchaseOffers.Num());
+	}
+	
+	/**
+	 * Generate a final receipt for all purchases made in this single transaction
+	 */
+	TSharedRef<FPurchaseReceipt> GenerateReceipt();
+	
+	/** Generate one off receipts for transactions initiated outside the current run of the application */
+	static TSharedRef<FPurchaseReceipt> GenerateReceipt(EPurchaseTransactionState Result, const FStoreKitTransactionData& Transaction);
+	
+	/** Mark this pending purchase as started */
+	void StartProcessing();
+	/** Add a single completed transaction (one of possibly many) to this transaction */
+	bool AddCompletedOffer(EPurchaseTransactionState Result, const FStoreKitTransactionData& Transaction);
+	/** @return true if all offers to purchase have completed transactions */
+	bool AreAllOffersComplete() const;
+	/** @return result of all purchases, if any fail it will return failed, if any cancel it will return canceled */
+	EPurchaseTransactionState GetFinalTransactionState() const;
+	
+public:
+	
+	/** Checkout info for the pending order */
+	const FPurchaseCheckoutRequest CheckoutRequest;
+	/** Mirror array of purchase state for the various offers to purchase */
+	TArray<EPurchaseTransactionState> OfferPurchaseStates;
+	/** User for the pending order */
+	const FUniqueNetIdString UserId;
+	/** Delegate to call on completion */
+	const FOnPurchaseCheckoutComplete CheckoutCompleteDelegate;
+	/** Tracks the current state of the order */
+	FPurchaseReceipt PendingPurchaseInfo;
+};
+
+/**
+ * Implementation for online purchase via IOS services
+ */
+class FOnlinePurchaseIOS :
+	public IOnlinePurchase,
+	public TSharedFromThis<FOnlinePurchaseIOS, ESPMode::ThreadSafe>
+{
+
+public:
+
+	// IOnlinePurchase
+
+	virtual bool IsAllowedToPurchase(const FUniqueNetId& UserId) override;
+	virtual void Checkout(const FUniqueNetId& UserId, const FPurchaseCheckoutRequest& CheckoutRequest, const FOnPurchaseCheckoutComplete& Delegate) override;
+	virtual void FinalizePurchase(const FUniqueNetId& UserId, const FString& ReceiptId) override;
+	virtual void RedeemCode(const FUniqueNetId& UserId, const FRedeemCodeRequest& RedeemCodeRequest, const FOnPurchaseRedeemCodeComplete& Delegate) override;
+	virtual void QueryReceipts(const FUniqueNetId& UserId, bool bRestoreReceipts, const FOnQueryReceiptsComplete& Delegate) override;
+	virtual void GetReceipts(const FUniqueNetId& UserId, TArray<FPurchaseReceipt>& OutReceipts) const override;
+	
+	// FOnlinePurchaseIOS
+
+	/**
+	 * Constructor
+	 *
+	 * @param InSubsystem IOS subsystem being used
+	 */
+	FOnlinePurchaseIOS(class FOnlineSubsystemIOS* InSubsystem);
+
+	/**
+	 * Constructor
+	 */
+	FOnlinePurchaseIOS();
+
+	/**
+	 * Destructor
+	 */
+	virtual ~FOnlinePurchaseIOS();
+	
+	/** Initialize the FStoreKitHelper for interaction with the app store */
+	void InitStoreKit(FStoreKitHelperV2* InStoreKit);
+	
+	
+private:
+	
+	/** delegate fired when a product request completes */
+	void OnProductPurchaseRequestResponse(SKProductsResponse* Response, const FOnQueryOnlineStoreOffersComplete& CompletionDelegate);
+	/** delegate fired when a single purchase transaction has completed (may be a part of multiple requests at once) */
+	void OnTransactionCompleteResponse(EPurchaseTransactionState Result, const FStoreKitTransactionData& TransactionData);
+	/** delegate fired when a single transaction is restored (may be a part of many restored purchases) */
+	void OnTransactionRestored(const FStoreKitTransactionData& TransactionData);
+	/** delegate fired when all transactions have been restored */
+	void OnRestoreTransactionsComplete(EPurchaseTransactionState Result);
+	
+	void OnTransactionInProgress(const FStoreKitTransactionData& TransactionData);
+	void OnTransactionDeferred(const FStoreKitTransactionData& TransactionData);
+	
+private:
+	
+	/** Mapping from user id to pending transaction */
+	typedef TMap<FString, TSharedRef<FOnlinePurchasePendingTransactionIOS> > FOnlinePurchasePendingTransactionMap;
+
+	/** Mapping from user id to complete transactions */
+	typedef TMap<FString, TArray< TSharedRef<FPurchaseReceipt> > > FOnlinePurchaseCompleteTransactionsMap;
+	
+	/** Array of transactions completion indirectly (previous run, etc) */
+	typedef TArray< TSharedRef<FPurchaseReceipt> > FOnlineCompletedTransactions;
+	
+private:
+	
+	/** Store kit helper for interfacing with app store (owned by main OnlineSubsystem) */
+	FStoreKitHelperV2* StoreHelper;
+	
+	/** Are transactions current being restored */
+	bool bRestoringTransactions;
+	/** Transient delegate to fire when query receipts has completed, when restoring transactions */
+	FOnQueryReceiptsComplete QueryReceiptsComplete;
+	
+	/** Keeps track of pending user transactions */
+	FOnlinePurchasePendingTransactionMap PendingTransactions;
+	
+	/** Cache of completed transactions */
+	FOnlinePurchaseCompleteTransactionsMap CompletedTransactions;
+	
+	/** Cache of purchases completed outside the running instance */
+	FOnlineCompletedTransactions OfflineTransactions;
+	
+	/** Reference to the parent subsystem */
+	FOnlineSubsystemIOS* Subsystem;
+};
+
+typedef TSharedPtr<FOnlinePurchaseIOS, ESPMode::ThreadSafe> FOnlinePurchaseIOSPtr;
\ No newline at end of file
diff --git a/Engine/Plugins/Online/IOS/OnlineSubsystemIOS/Source/Public/OnlineStoreIOS.h b/Engine/Plugins/Online/IOS/OnlineSubsystemIOS/Source/Public/OnlineStoreIOS.h
new file mode 100644
--- /dev/null
+++ b/Engine/Plugins/Online/IOS/OnlineSubsystemIOS/Source/Public/OnlineStoreIOS.h
@@ -0,0 +1,149 @@
+// Copyright 1998-2016 Epic Games, Inc. All Rights Reserved.
+
+#pragma once
+
+#include "OnlineStoreInterfaceV2.h"
+
+enum class EPurchaseTransactionState : uint8;
+struct FStoreKitTransactionData;
+@class FStoreKitHelperV2;
+@class SKProductsResponse;
+@class SKProduct;
+
+/**
+ * Implementation for online store via IOS services
+ */
+class FOnlineStoreIOS :
+	public IOnlineStoreV2,
+	public TSharedFromThis<FOnlineStoreIOS, ESPMode::ThreadSafe>
+{
+public:
+
+	// IOnlineStoreV2
+
+	virtual void QueryCategories(const FUniqueNetId& UserId, const FOnQueryOnlineStoreCategoriesComplete& Delegate) override;
+	virtual void GetCategories(TArray<FOnlineStoreCategory>& OutCategories) const override;
+	virtual void QueryOffersByFilter(const FUniqueNetId& UserId, const FOnlineStoreFilter& Filter, const FOnQueryOnlineStoreOffersComplete& Delegate) override;
+	virtual void QueryOffersById(const FUniqueNetId& UserId, const TArray<FUniqueOfferId>& OfferIds, const FOnQueryOnlineStoreOffersComplete& Delegate) override;
+	virtual void GetOffers(TArray<FOnlineStoreOfferRef>& OutOffers) const override;
+	virtual TSharedPtr<FOnlineStoreOffer> GetOffer(const FUniqueOfferId& OfferId) const override;
+
+	// FOnlineStoreIOS
+
+	/**
+	 * Constructor
+	 *
+	 * @param InSubsystem Online subsystem being used
+	 */
+	FOnlineStoreIOS(class FOnlineSubsystemIOS* InSubsystem);
+
+	/**
+	 * Constructor
+	 */
+	FOnlineStoreIOS();
+
+	/**
+	 * Destructor
+	 */
+	virtual ~FOnlineStoreIOS();
+	
+	/** Initialize the FStoreKitHelper for interaction with the app store */
+	void InitStoreKit(FStoreKitHelperV2* InStoreKit);
+	
+	/**
+	 * Get the production information for a given offer id
+	 * Must have previously been retrieved vi QueryOffers*
+	 *
+	 * @return the SKProduct previously queried for product information
+	 */
+	SKProduct* GetSKProductByOfferId(const FUniqueOfferId& OfferId);
+
+private:
+	
+	/**
+	 * Representation of a single product offer
+	 */
+	struct FOnlineStoreOfferIOS
+	{
+		FOnlineStoreOfferIOS()
+		: Product(nil)
+		, Offer(nullptr)
+		{
+		}
+		
+		FOnlineStoreOfferIOS(SKProduct* InProduct, const FOnlineStoreOffer& InOffer)
+		{
+			Product = InProduct;
+			[Product retain];
+			
+			Offer = MakeShareable(new FOnlineStoreOffer(InOffer));
+		}
+		
+		FOnlineStoreOfferIOS(const FOnlineStoreOfferIOS& Other)
+		{
+			Product = Other.Product;
+			[Product retain];
+			Offer = Other.Offer;
+		}
+		
+		FOnlineStoreOfferIOS& operator=(const FOnlineStoreOfferIOS& Other)
+		{
+			// Free
+			[Product release];
+
+			// Copy
+			Product = Other.Product;
+			[Product retain];
+			Offer = Other.Offer;
+			return *this;
+		}
+		
+		~FOnlineStoreOfferIOS()
+		{
+			[Product release];
+		}
+		
+		/** @return true if the store offer is valid/proper */
+		bool IsValid() const
+		{
+			return Product != nil && Offer.IsValid();
+		}
+		
+		/** Reference to the app store product information */
+		SKProduct* Product;
+		/** Product information about this offer */
+		TSharedPtr<FOnlineStoreOffer> Offer;
+	};
+
+	void AddOffer(const FOnlineStoreOfferIOS& NewOffer);
+
+private:
+
+	/** delegate fired when a product request completes */
+	void OnProductPurchaseRequestResponse(SKProductsResponse* Response, const FOnQueryOnlineStoreOffersComplete& CompletionDelegate);
+	/** delegate fired when a single purchase transaction has completed (may be a part of multiple requests at once) */
+	void OnTransactionCompleteResponse(EPurchaseTransactionState Result, const FStoreKitTransactionData& TransactionData);
+	/** delegate fired when a single transaction is restored (may be a part of many restored purchases) */
+	void OnTransactionRestored(const FStoreKitTransactionData& TransactionData);
+	/** delegate fired when all transactions have been restored */
+	void OnRestoreTransactionsComplete(EPurchaseTransactionState Result);
+	
+private:
+	
+	/** Mapping from offer id to product information */
+	typedef TMap<FUniqueOfferId, FOnlineStoreOfferIOS> FOnlineOfferDescriptionMap;
+
+	/** Mapping of all queried offers to their product information */
+	FOnlineOfferDescriptionMap CachedOffers;
+
+	/** Store kit helper for interfacing with app store */
+	FStoreKitHelperV2* StoreHelper;
+
+	/** Is a query already in flight */
+	bool bIsQueryInFlight;
+	
+	/** Reference to the parent subsystem */
+	FOnlineSubsystemIOS* Subsystem;
+};
+
+typedef TSharedPtr<FOnlineStoreIOS, ESPMode::ThreadSafe> FOnlineStoreIOSPtr;
\ No newline at end of file
diff --git a/Engine/Plugins/Online/IOS/OnlineSubsystemIOS/Source/Public/OnlineSubsystemIOS.h b/Engine/Plugins/Online/IOS/OnlineSubsystemIOS/Source/Public/OnlineSubsystemIOS.h
--- a/Engine/Plugins/Online/IOS/OnlineSubsystemIOS/Source/Public/OnlineSubsystemIOS.h
+++ b/Engine/Plugins/Online/IOS/OnlineSubsystemIOS/Source/Public/OnlineSubsystemIOS.h
@@ -6,6 +6,8 @@
 #include "OnlineSubsystemImpl.h"
 #include "OnlineSubsystemIOSPackage.h"
 
+@class FStoreKitHelperV2;
+
 
 /**
  *	OnlineSubsystemIOS - Implementation of the online subsystem for IOS services
@@ -33,8 +35,8 @@ public:
 	virtual IOnlineTitleFilePtr GetTitleFileInterface() const override;
 	virtual IOnlineEntitlementsPtr GetEntitlementsInterface() const override;
 	virtual IOnlineStorePtr GetStoreInterface() const override;
-	virtual IOnlineStoreV2Ptr GetStoreV2Interface() const override { return nullptr; }
-	virtual IOnlinePurchasePtr GetPurchaseInterface() const override { return nullptr; }
+	virtual IOnlineStoreV2Ptr GetStoreV2Interface() const override;
+	virtual IOnlinePurchasePtr GetPurchaseInterface() const override;
 	virtual IOnlineEventsPtr GetEventsInterface() const override;
 	virtual IOnlineAchievementsPtr GetAchievementsInterface() const override;
 	virtual IOnlineSharingPtr GetSharingInterface() const override;
@@ -59,15 +61,27 @@ PACKAGE_SCOPE:
 	bool IsEnabled();
 
 	/**
-	* Is IAP available for use
-	* @return true if IAP should be available, false otherwise
-	*/
+	 * Is IAP available for use
+	 * @return true if IAP should be available, false otherwise
+	 */
 	bool IsInAppPurchasingEnabled();
+	
+	/**
+	 * Is Store v2 enabled (disabling legacy store interface)
+	 * @return true if enabled, false otherwise
+	 */
+	bool IsV2StoreEnabled();
 
 private:
+	
+	void InitStoreKitHelper();
+	void CleanupStoreKitHelper();
 
 	/** Online async task thread */
 	class FRunnableThread* OnlineAsyncTaskThread;
+	
+	/** Store kit helper for interfacing with app store */
+	FStoreKitHelperV2* StoreHelper;
 
 	/** Interface to the session services */
 	FOnlineSessionIOSPtr SessionInterface;
@@ -83,6 +97,12 @@ private:
 
 	/** Interface to the online store */
 	FOnlineStoreInterfaceIOSPtr StoreInterface;
+	
+	/** Interface to the online catalog */
+	FOnlineStoreIOSPtr StoreV2Interface;
+	
+	/** Interface to the store purchasing */
+	FOnlinePurchaseIOSPtr PurchaseInterface;
 
 	/** Interface to the online achievements */
 	FOnlineAchievementsIOSPtr AchievementsInterface;
diff --git a/Engine/Plugins/Online/OnlineFramework/OnlineFramework.uplugin b/Engine/Plugins/Online/OnlineFramework/OnlineFramework.uplugin
--- a/Engine/Plugins/Online/OnlineFramework/OnlineFramework.uplugin
+++ b/Engine/Plugins/Online/OnlineFramework/OnlineFramework.uplugin
@@ -2,7 +2,7 @@
 	"FileVersion" : 3,
 	
 	"FriendlyName" : "Online Framework Plugin",
-	"Version" : 1,
+	"Version" : 2,
 	"VersionName" : "1.0",
 	"EngineVersion" : "4.2.0",
 	"Description" : "Shared code for interacting with online gameplay services.",
@@ -27,6 +27,10 @@
         {
             "Name": "Hotfix",
             "Type": "Runtime"
+        },
+        {
+            "Name": "Rejoin",
+            "Type": "Runtime"
         }
     ]
 }
\ No newline at end of file
diff --git a/Engine/Plugins/Online/OnlineFramework/Source/Hotfix/Hotfix.Build.cs b/Engine/Plugins/Online/OnlineFramework/Source/Hotfix/Hotfix.Build.cs
--- a/Engine/Plugins/Online/OnlineFramework/Source/Hotfix/Hotfix.Build.cs
+++ b/Engine/Plugins/Online/OnlineFramework/Source/Hotfix/Hotfix.Build.cs
@@ -6,6 +6,8 @@ public class Hotfix : ModuleRules
 {
 	public Hotfix(TargetInfo Target)
     {
+		PCHUsage = ModuleRules.PCHUsageMode.UseExplicitOrSharedPCHs;
+		
         PrivateDependencyModuleNames.AddRange(
 			new string[] { 
 				"Core",
diff --git a/Engine/Plugins/Online/OnlineFramework/Source/Hotfix/Private/OnlineHotfixManager.cpp b/Engine/Plugins/Online/OnlineFramework/Source/Hotfix/Private/OnlineHotfixManager.cpp
--- a/Engine/Plugins/Online/OnlineFramework/Source/Hotfix/Private/OnlineHotfixManager.cpp
+++ b/Engine/Plugins/Online/OnlineFramework/Source/Hotfix/Private/OnlineHotfixManager.cpp
@@ -83,6 +83,8 @@ void UOnlineHotfixManager::PostInitProperties()
 	// So we only try to apply files for this platform
 	PlatformPrefix = DebugPrefix + ANSI_TO_TCHAR(FPlatformProperties::PlatformName());
 	PlatformPrefix += TEXT("_");
+	// Server prefix
+	ServerPrefix = DebugPrefix + TEXT("DedicatedServer");
 	// Build the default prefix too
 	DefaultPrefix = DebugPrefix + TEXT("Default");
 
@@ -160,60 +162,67 @@ struct FHotfixFileSortPredicate
 	struct FHotfixFileNameSortPredicate
 	{
 		const FString PlatformPrefix;
+		const FString ServerPrefix;
 		const FString DefaultPrefix;
 
-		FHotfixFileNameSortPredicate(const FString& InPlatformPrefix, const FString& InDefaultPrefix) :
+		FHotfixFileNameSortPredicate(const FString& InPlatformPrefix, const FString& InServerPrefix, const FString& InDefaultPrefix) :
 			PlatformPrefix(InPlatformPrefix),
+			ServerPrefix(InServerPrefix),
 			DefaultPrefix(InDefaultPrefix)
 		{
 		}
 
-		bool Compare(const FString& A, const FString& B) const
+		int32 GetPriorityForCompare(const FString& HotfixName) const
 		{
-			bool bAIsIni = A.EndsWith(TEXT("INI"));
-			bool bBIsIni = B.EndsWith(TEXT("INI"));
-			if (bAIsIni && bBIsIni)
+			// Non-ini files are applied last
+			int32 Priority = 5;
+			
+			if (HotfixName.EndsWith(TEXT("INI")))
 			{
-				bool bAStartsWithDefault = A.StartsWith(DefaultPrefix);
-				bool bBStartsWithDefault = B.StartsWith(DefaultPrefix);
-				bool bAStartsWithPlatformName = A.StartsWith(PlatformPrefix);
-				bool bBStartsWithPlatformName = B.StartsWith(PlatformPrefix);
-				// Sort any file with Default in front of the Platform version so they apply in correct order
-				if (bAStartsWithDefault && bBStartsWithPlatformName)
+				// Defaults are applied first
+				if (HotfixName.StartsWith(DefaultPrefix))
 				{
-					return true;
+					Priority = 1;
 				}
-				else if (bAStartsWithPlatformName && bBStartsWithDefault)
+				// Server trumps default
+				else if (HotfixName.StartsWith(ServerPrefix))
 				{
-					return false;
+					Priority = 2;
 				}
-				// Sort any Default/Platform INIs before other INIs
-				else if ((bAStartsWithDefault || bAStartsWithPlatformName) && !(bBStartsWithDefault || bBStartsWithPlatformName))
+				// Platform trumps server
+				else if (HotfixName.StartsWith(PlatformPrefix))
 				{
-					return true;
+					Priority = 3;
 				}
-				else if (!(bAStartsWithDefault || bAStartsWithPlatformName) && (bBStartsWithDefault || bBStartsWithPlatformName))
+				// Other INIs whitelisted in game override of WantsHotfixProcessing will trump all other INIs
+				else
 				{
-					return false;
+					Priority = 4;
 				}
 			}
-			// Sort all INI files before non-INI files
-			else if (bAIsIni && !bBIsIni)
+
+			return Priority;
+		}
+
+		bool Compare(const FString& A, const FString& B) const
+		{
+			int32 APriority = GetPriorityForCompare(A);
+			int32 BPriority = GetPriorityForCompare(B);
+			if (APriority != BPriority)
 			{
-				return true;
+				return APriority < BPriority;
 			}
-			else if (!bAIsIni && bBIsIni)
+			else
 			{
-				return false;
+				// Fall back to sort by the string order if both have same priority
+				return A < B;
 			}
-			// Sort by the string order
-			return A < B;
 		}
 	};
 	FHotfixFileNameSortPredicate FileNameSorter;
 
-	FHotfixFileSortPredicate(const FString& InPlatformPrefix, const FString& InDefaultPrefix) :
-		FileNameSorter(InPlatformPrefix, InDefaultPrefix)
+	FHotfixFileSortPredicate(const FString& InPlatformPrefix, const FString& InServerPrefix, const FString& InDefaultPrefix) :
+		FileNameSorter(InPlatformPrefix, InServerPrefix, InDefaultPrefix)
 	{
 	}
 
@@ -242,7 +251,7 @@ void UOnlineHotfixManager::OnEnumerateFilesComplete(bool bWasSuccessful)
 		// Reduce the set of work to just the files that changed since last run
 		BuildHotfixFileListDeltas();
 		// Sort after filtering so that the comparison below doesn't fail to different order from the server
-		ChangedHotfixFileList.Sort<FHotfixFileSortPredicate>(FHotfixFileSortPredicate(PlatformPrefix, DefaultPrefix));
+		ChangedHotfixFileList.Sort<FHotfixFileSortPredicate>(FHotfixFileSortPredicate(PlatformPrefix, ServerPrefix, DefaultPrefix));
 		// Perform any undo operations needed
 		if (ChangedHotfixFileList.Num() > 0 || RemovedHotfixFileList.Num() > 0)
 		{
@@ -520,7 +529,25 @@ bool UOnlineHotfixManager::WantsHotfixProcessing(const FCloudFileHeader& FileHea
 	const FString Extension = FPaths::GetExtension(FileHeader.FileName);
 	if (Extension == TEXT("INI"))
 	{
-		return FileHeader.FileName.StartsWith(PlatformPrefix) || FileHeader.FileName.StartsWith(DefaultPrefix);
+		bool bIsServerHotfix = FileHeader.FileName.StartsWith(ServerPrefix);
+		bool bWantsServerHotfix = IsRunningDedicatedServer() && bIsServerHotfix;
+		bool bWantsDefaultHotfix = FileHeader.FileName.StartsWith(DefaultPrefix);
+		bool bWantsPlatformHotfix = FileHeader.FileName.StartsWith(PlatformPrefix);
+
+		if (bWantsPlatformHotfix)
+		{
+			UE_LOG(LogHotfixManager, Verbose, TEXT("Using platform hotfix %s"), * FileHeader.FileName);
+		}
+		else if (bWantsServerHotfix)
+		{
+			UE_LOG(LogHotfixManager, Verbose, TEXT("Using server hotfix %s"), * FileHeader.FileName);
+		}
+		else if (bWantsDefaultHotfix)
+		{
+			UE_LOG(LogHotfixManager, Verbose, TEXT("Using default hotfix %s"), * FileHeader.FileName);
+		}
+
+		return bWantsPlatformHotfix || bWantsServerHotfix || bWantsDefaultHotfix;
 	}
 	else if (Extension == TEXT("PAK"))
 	{
@@ -753,7 +780,7 @@ bool UOnlineHotfixManager::HotfixPakFile(const FCloudFileHeader& FileHeader)
 			}
 		}
 		// Sort the INIs so they are processed consistently
-		IniList.Sort<FHotfixFileSortPredicate>(FHotfixFileSortPredicate(PlatformPrefix, DefaultPrefix));
+		IniList.Sort<FHotfixFileSortPredicate>(FHotfixFileSortPredicate(PlatformPrefix, ServerPrefix, DefaultPrefix));
 		// Now process the INIs in sorted order
 		for (auto& IniName : IniList)
 		{
@@ -990,6 +1017,8 @@ struct FHotfixManagerExec :
 			FCloudFileHeader Header;
 			Header.FileName = TEXT("SomeRandom.ini");
 			TestList.Add(Header);
+			Header.FileName = TEXT("DedicatedServerGame.ini");
+			TestList.Add(Header);
 			Header.FileName = TEXT("pakchunk1-PS4_P.pak");
 			TestList.Add(Header);
 			Header.FileName = TEXT("EN_Game.locres");
@@ -1006,7 +1035,9 @@ struct FHotfixManagerExec :
 			TestList.Add(Header);
 			Header.FileName = TEXT("AnotherRandom.ini");
 			TestList.Add(Header);
-			TestList.Sort<FHotfixFileSortPredicate>(FHotfixFileSortPredicate(TEXT("PS4_"), TEXT("Default")));
+			Header.FileName = TEXT("DedicatedServerEngine.ini");
+			TestList.Add(Header);
+			TestList.Sort<FHotfixFileSortPredicate>(FHotfixFileSortPredicate(TEXT("PS4_"), TEXT("DedicatedServer"), TEXT("Default")));
 
 			UE_LOG(LogHotfixManager, Log, TEXT("Hotfixing sort is:"));
 			for (auto& FileHeader : TestList)
@@ -1018,10 +1049,12 @@ struct FHotfixManagerExec :
 			TestList2.Add(TEXT("SomeRandom.ini"));
 			TestList2.Add(TEXT("DefaultGame.ini"));
 			TestList2.Add(TEXT("PS4_DefaultEngine.ini"));
+			TestList2.Add(TEXT("DedicatedServerEngine.ini"));
+			TestList2.Add(TEXT("DedicatedServerGame.ini"));
 			TestList2.Add(TEXT("DefaultEngine.ini"));
 			TestList2.Add(TEXT("PS4_DefaultGame.ini"));
 			TestList2.Add(TEXT("AnotherRandom.ini"));
-			TestList2.Sort<FHotfixFileSortPredicate>(FHotfixFileSortPredicate(TEXT("PS4_"), TEXT("Default")));
+			TestList2.Sort<FHotfixFileSortPredicate>(FHotfixFileSortPredicate(TEXT("PS4_"), TEXT("DedicatedServer"), TEXT("Default")));
 
 			UE_LOG(LogHotfixManager, Log, TEXT("Hotfixing PAK INI file sort is:"));
 			for (auto& IniName : TestList2)
diff --git a/Engine/Plugins/Online/OnlineFramework/Source/Hotfix/Private/UpdateManager.cpp b/Engine/Plugins/Online/OnlineFramework/Source/Hotfix/Private/UpdateManager.cpp
--- a/Engine/Plugins/Online/OnlineFramework/Source/Hotfix/Private/UpdateManager.cpp
+++ b/Engine/Plugins/Online/OnlineFramework/Source/Hotfix/Private/UpdateManager.cpp
@@ -101,9 +101,10 @@ UUpdateManager::EUpdateStartResult UUpdateManager::StartCheckInternal(bool bInCh
 		return Result;
 	}
 
-	if (CurrentUpdateState == EUpdateState::UpdateIdle ||
+	if (!IsTimerHandleActive(StartCheckInternalTimerHandle) &&
+		(CurrentUpdateState == EUpdateState::UpdateIdle ||
 		CurrentUpdateState == EUpdateState::UpdatePending ||
-		CurrentUpdateState == EUpdateState::UpdateComplete)
+		CurrentUpdateState == EUpdateState::UpdateComplete))
 	{
 		bCheckHotfixAvailabilityOnly = bInCheckHotfixOnly;
 
@@ -129,7 +130,7 @@ UUpdateManager::EUpdateStartResult UUpdateManager::StartCheckInternal(bool bInCh
 			};
 
 			// Give the UI state widget a chance to start listening for delegates
-			DelayResponse(StartDelegate, 0.2f);
+			StartCheckInternalTimerHandle = DelayResponse(StartDelegate, 0.2f);
 			Result = EUpdateStartResult::UpdateStarted;
 		}
 		else
@@ -140,7 +141,7 @@ UUpdateManager::EUpdateStartResult UUpdateManager::StartCheckInternal(bool bInCh
 				CheckComplete(LastResult, false);
 			};
 
-			DelayResponse(StartDelegate, 0.1f);
+			StartCheckInternalTimerHandle = DelayResponse(StartDelegate, 0.1f);
 			Result = EUpdateStartResult::UpdateCached;
 		}
 	}
@@ -406,8 +407,16 @@ void UUpdateManager::PlatformEnvironmentCheck_OnLoginConsoleComplete(int32 Local
 	}
 	else
 	{
-		UE_LOG(LogHotfixManager, Warning, TEXT("Failed to detect online environment for the platform"));
-		CheckComplete(EUpdateCompletionStatus::UpdateFailure_NotLoggedIn);
+		if (Error.Contains(TEXT("getUserAccessCode failed : 0x8055000f"), ESearchCase::IgnoreCase))
+		{
+			UE_LOG(LogHotfixManager, Warning, TEXT("Failed to complete login because patch is required"));
+			CheckComplete(EUpdateCompletionStatus::UpdateSuccess_NeedsPatch);
+		}
+		else
+		{
+			UE_LOG(LogHotfixManager, Warning, TEXT("Failed to detect online environment for the platform"));
+			CheckComplete(EUpdateCompletionStatus::UpdateFailure_NotLoggedIn);
+		}
 	}
 }
 
@@ -638,6 +647,17 @@ FTimerHandle UUpdateManager::DelayResponse(DelayCb&& Delegate, float Delay)
 	return TimerHandle;
 }
 
+bool UUpdateManager::IsTimerHandleActive(const FTimerHandle& TimerHandle) const
+{
+	UWorld* World = GetWorld();
+	if (ensure(World != nullptr))
+	{
+		return World->GetTimerManager().IsTimerActive(TimerHandle);
+	}
+
+	return false;
+}
+
 UWorld* UUpdateManager::GetWorld() const
 {
 	UGameInstance* GameInstance = GetTypedOuter<UGameInstance>();
diff --git a/Engine/Plugins/Online/OnlineFramework/Source/Hotfix/Public/OnlineHotfixManager.h b/Engine/Plugins/Online/OnlineFramework/Source/Hotfix/Public/OnlineHotfixManager.h
--- a/Engine/Plugins/Online/OnlineFramework/Source/Hotfix/Public/OnlineHotfixManager.h
+++ b/Engine/Plugins/Online/OnlineFramework/Source/Hotfix/Public/OnlineHotfixManager.h
@@ -141,6 +141,8 @@ protected:
 	FString GameLocName;
 	/** Used to match any PAK files for this platform */
 	FString PlatformPrefix;
+	/** Used to match any server-only hotfixes */
+	FString ServerPrefix;
 	/** Normally will be "Default" but could be different if we have a debug prefix */
 	FString DefaultPrefix;
 	/** Holds a chunk of string that will be swapped for Game during processing pak files (MyGame/Content/Maps -> /Game/Maps) */
diff --git a/Engine/Plugins/Online/OnlineFramework/Source/Hotfix/Public/UpdateManager.h b/Engine/Plugins/Online/OnlineFramework/Source/Hotfix/Public/UpdateManager.h
--- a/Engine/Plugins/Online/OnlineFramework/Source/Hotfix/Public/UpdateManager.h
+++ b/Engine/Plugins/Online/OnlineFramework/Source/Hotfix/Public/UpdateManager.h
@@ -187,6 +187,11 @@ public:
 	 */
 	bool IsUpdating() const { return !bCheckHotfixAvailabilityOnly && (CurrentUpdateState != EUpdateState::UpdateIdle) && (CurrentUpdateState != EUpdateState::UpdateComplete); }
 
+	/**
+	* @return true if the update manager is actively checking
+	*/
+	bool IsActivelyUpdating() const { return IsUpdating() && (CurrentUpdateState != EUpdateState::UpdatePending); }
+
 	/**
 	 * @return the current state of the update check process
 	 */
@@ -354,6 +359,7 @@ protected:
 	EUpdateCompletionStatus LastCompletionResult[2];
 
 	FDelegateHandle TickerHandle;
+	FTimerHandle StartCheckInternalTimerHandle;
 
 private:
 
@@ -371,6 +377,9 @@ private:
 	typedef TFunction<void(void)> DelayCb;
 	FTimerHandle DelayResponse(DelayCb&& Delegate, float Delay);
 
+	/** Helper function to check if a timer is running. */
+	bool IsTimerHandleActive(const FTimerHandle& TimerHandle) const;
+
 	friend bool SkipPatchCheck(UUpdateManager* UpdateManager);
 
 protected:
diff --git a/Engine/Plugins/Online/OnlineFramework/Source/Lobby/Lobby.Build.cs b/Engine/Plugins/Online/OnlineFramework/Source/Lobby/Lobby.Build.cs
--- a/Engine/Plugins/Online/OnlineFramework/Source/Lobby/Lobby.Build.cs
+++ b/Engine/Plugins/Online/OnlineFramework/Source/Lobby/Lobby.Build.cs
@@ -5,8 +5,9 @@ using UnrealBuildTool;
 public class Lobby : ModuleRules
 {
 	public Lobby(TargetInfo Target)
-	{
+	{		
 		Definitions.Add("LOBBY_PACKAGE=1");
+		PCHUsage = ModuleRules.PCHUsageMode.UseExplicitOrSharedPCHs;
 
 		PrivateIncludePaths.AddRange(
 			new string[] {
diff --git a/Engine/Plugins/Online/OnlineFramework/Source/Lobby/Private/LobbyBeaconClient.cpp b/Engine/Plugins/Online/OnlineFramework/Source/Lobby/Private/LobbyBeaconClient.cpp
--- a/Engine/Plugins/Online/OnlineFramework/Source/Lobby/Private/LobbyBeaconClient.cpp
+++ b/Engine/Plugins/Online/OnlineFramework/Source/Lobby/Private/LobbyBeaconClient.cpp
@@ -15,7 +15,8 @@ ALobbyBeaconClient::ALobbyBeaconClient(const FObjectInitializer& ObjectInitializ
 	Super(ObjectInitializer),
 	LobbyState(nullptr),
 	PlayerState(nullptr),
-	bLoggedIn(false)
+	bLoggedIn(false),
+	bLobbyJoinAcked(false)
 {
 	bOnlyRelevantToOwner = true;
 }
@@ -125,6 +126,7 @@ void ALobbyBeaconClient::JoiningServer()
 	if (bLoggedIn)
 	{
 		UE_LOG(LogBeacon, Log, TEXT("JoiningServer %s Id: %s"), *GetName(), PlayerState ? *PlayerState->UniqueId->ToString() : TEXT("Unknown"));
+		bLobbyJoinAcked = false;
 		ServerNotifyJoiningServer();
 	}
 	else
@@ -265,6 +267,7 @@ void ALobbyBeaconClient::AckJoiningServer()
 void ALobbyBeaconClient::ClientAckJoiningServer_Implementation()
 {
 	UE_LOG(LogBeacon, Log, TEXT("ClientAckJoiningServer %s Id: %s LoggedIn: %d"), *GetName(), PlayerState ? *PlayerState->UniqueId->ToString() : TEXT("Unknown"), bLoggedIn);
+	bLobbyJoinAcked = true;
 	OnJoiningGameAck().ExecuteIfBound();
 }
 
diff --git a/Engine/Plugins/Online/OnlineFramework/Source/Lobby/Public/LobbyBeaconClient.h b/Engine/Plugins/Online/OnlineFramework/Source/Lobby/Public/LobbyBeaconClient.h
--- a/Engine/Plugins/Online/OnlineFramework/Source/Lobby/Public/LobbyBeaconClient.h
+++ b/Engine/Plugins/Online/OnlineFramework/Source/Lobby/Public/LobbyBeaconClient.h
@@ -128,6 +128,9 @@ protected:
 	/** Has this beacon been properly logged in */
 	bool bLoggedIn;
 
+	/** True once the server has acknowledged our join intent */
+	bool bLobbyJoinAcked;
+
 	/** Session Id of the destination host */
 	FString DestSessionId;
 	
diff --git a/Engine/Plugins/Online/OnlineFramework/Source/Party/Private/Chatroom.cpp b/Engine/Plugins/Online/OnlineFramework/Source/Party/Private/Chatroom.cpp
--- a/Engine/Plugins/Online/OnlineFramework/Source/Party/Private/Chatroom.cpp
+++ b/Engine/Plugins/Online/OnlineFramework/Source/Party/Private/Chatroom.cpp
@@ -39,7 +39,7 @@ void UChatroom::UnregisterDelegates()
 	}
 }
 
-void UChatroom::CreateOrJoinChatRoom(const FUniqueNetIdRepl& LocalUserId, const FChatRoomId& ChatRoomId, const FOnChatRoomCreatedOrJoined& CompletionDelegate, const FString& Password)
+void UChatroom::CreateOrJoinChatRoom(FUniqueNetIdRepl LocalUserId, FChatRoomId ChatRoomId, FOnChatRoomCreatedOrJoined CompletionDelegate, FString Password)
 {
 	bool bShouldFireDelegate = false;
 	if (!ChatRoomId.IsEmpty())
@@ -48,7 +48,7 @@ void UChatroom::CreateOrJoinChatRoom(const FUniqueNetIdRepl& LocalUserId, const
 		{
 			UWorld* World = GetWorld();
 			IOnlineChatPtr ChatInt = Online::GetChatInterface(World);
-			if (ensure(ChatInt.IsValid()))
+			if (ChatInt.IsValid())
 			{
 				if (!IsAlreadyInChatRoom(LocalUserId, ChatRoomId))
 				{
@@ -56,7 +56,12 @@ void UChatroom::CreateOrJoinChatRoom(const FUniqueNetIdRepl& LocalUserId, const
 					{
 						UE_LOG(LogOnlineChat, Verbose, TEXT("[UChatroom::CreateOrJoinChatRoom] attempting to join %s"), *ChatRoomId);
 						CurrentChatRoomId = ChatRoomId;
-						GetTimerManager().ClearTimer(ChatRoomRetryTimerHandle);
+
+						FTimerManager& TM = GetTimerManager();
+						if (TM.GetTimerRemaining(ChatRoomRetryTimerHandle) > 0)
+						{
+							TM.ClearTimer(ChatRoomRetryTimerHandle);
+						}
 
 						FOnChatRoomJoinPublicDelegate ChatRoomDelegate;
 						ChatRoomDelegate.BindUObject(this, &ThisClass::OnChatRoomCreatedOrJoined, CompletionDelegate, Password);
@@ -182,7 +187,7 @@ void UChatroom::LeaveChatRoom(const FUniqueNetIdRepl& LocalUserId, const FOnChat
 		ChatRoomRetryTimerHandle.Invalidate();
 		NumChatRoomRetries = 0;
 
-		// ExitRoom callback could fire inline if it fails, make a copy so the ensure in OnChatRoomLeft does not hit
+		// ExitRoom callback could fire inline if it fails, make a copy so the ensure in ChatRoomLeftInternal does not hit
 		FChatRoomId ChatRoomIdCpy = CurrentChatRoomId;
 		CurrentChatRoomId.Empty();
 
@@ -203,7 +208,7 @@ void UChatroom::LeaveChatRoom(const FUniqueNetIdRepl& LocalUserId, const FOnChat
 		{
 			UE_LOG(LogOnlineChat, Display, TEXT("[UChatroom::LeaveChatRoom] Left chat while not logged in"));
 			// Clear out the chat room to avoid ensures (logged out, not in this room anymore)
-			OnChatRoomLeft(*LocalUserId, ChatRoomIdCpy, true, FString(), ChatRoomIdCpy, CompletionDelegate);
+			ChatRoomLeftInternal(ChatRoomIdCpy, CompletionDelegate);
 		}
 	}
 	else
@@ -225,13 +230,7 @@ void UChatroom::OnChatRoomLeft(const FUniqueNetId& LocalUserId, const FChatRoomI
 			ChatInt->ClearOnChatRoomExitDelegate_Handle(ChatRoomLeftDelegateHandle);
 		}
 
-		ensure(CurrentChatRoomId.IsEmpty());
-
-		GetTimerManager().SetTimerForNextTick(FTimerDelegate::CreateLambda([CompletionDelegate, RoomId]()
-		{
-			// Announce that chat is available
-			CompletionDelegate.ExecuteIfBound(RoomId);
-		}));
+		ChatRoomLeftInternal(RoomId, CompletionDelegate);
 	}
 	else
 	{
@@ -239,6 +238,17 @@ void UChatroom::OnChatRoomLeft(const FUniqueNetId& LocalUserId, const FChatRoomI
 	}
 }
 
+void UChatroom::ChatRoomLeftInternal(const FChatRoomId& RoomId, const FOnChatRoomLeft& CompletionDelegate)
+{
+	ensure(CurrentChatRoomId.IsEmpty());
+
+	GetTimerManager().SetTimerForNextTick(FTimerDelegate::CreateLambda([CompletionDelegate, RoomId]()
+	{
+		// Announce that chat is available
+		CompletionDelegate.ExecuteIfBound(RoomId);
+	}));
+}
+
 bool UChatroom::IsOnline() const
 {
 	return true;
@@ -269,9 +279,17 @@ bool UChatroom::IsAlreadyInChatRoom(const FUniqueNetIdRepl& LocalUserId, const F
 
 UWorld* UChatroom::GetWorld() const
 {
-	UWorld* World = GetOuter() ? GetOuter()->GetWorld() : nullptr;
-	checkf(World, TEXT("[UChatroom::GetWorld] Should have an outer that can access a world"));
-	return World;
+	if (HasAnyFlags(RF_ClassDefaultObject))
+	{
+		UE_LOG(LogOnlineChat, Warning, TEXT("[UChatroom::GetWorld] Called GetWorld on the CDO"));
+		return nullptr;
+	}
+	else
+	{
+		UWorld* World = GetOuter() ? GetOuter()->GetWorld() : nullptr;
+		checkf(World, TEXT("[UChatroom::GetWorld] Should have an outer that can access a world"));
+		return World;
+	}
 }
 
 FTimerManager& UChatroom::GetTimerManager() const
diff --git a/Engine/Plugins/Online/OnlineFramework/Source/Party/Private/Party.cpp b/Engine/Plugins/Online/OnlineFramework/Source/Party/Private/Party.cpp
--- a/Engine/Plugins/Online/OnlineFramework/Source/Party/Private/Party.cpp
+++ b/Engine/Plugins/Online/OnlineFramework/Source/Party/Private/Party.cpp
@@ -1207,15 +1207,16 @@ void UParty::LeaveAndRestorePersistentPartyInternal()
 	check(GameInstance);
 
 	TSharedPtr<const FUniqueNetId> PrimaryUserId = GameInstance->GetPrimaryPlayerUniqueId();
-	check(PrimaryUserId.IsValid() && PrimaryUserId->IsValid());
-
-	UPartyDelegates::FOnLeaveUPartyComplete CompletionDelegate;
-	CompletionDelegate.BindUObject(this, &ThisClass::OnLeavePersistentPartyAndRestore);
+	if (ensure(PrimaryUserId.IsValid() && PrimaryUserId->IsValid()))
+	{
+		UPartyDelegates::FOnLeaveUPartyComplete CompletionDelegate;
+		CompletionDelegate.BindUObject(this, &ThisClass::OnLeavePersistentPartyAndRestore);
 
-	// Unset this here, wouldn't be here if we were unable to call LeaveAndRestorePersistentParty, it will be reset inside
-	ensure(bLeavingPersistentParty);
-	bLeavingPersistentParty = false;
-	LeavePersistentParty(*PrimaryUserId, CompletionDelegate);
+		// Unset this here, wouldn't be here if we were unable to call LeaveAndRestorePersistentParty, it will be reset inside
+		ensure(bLeavingPersistentParty);
+		bLeavingPersistentParty = false;
+		LeavePersistentParty(*PrimaryUserId, CompletionDelegate);
+	}
 }
 
 void UParty::OnLeavePersistentPartyAndRestore(const FUniqueNetId& LocalUserId, const ELeavePartyCompletionResult Result)
diff --git a/Engine/Plugins/Online/OnlineFramework/Source/Party/Public/Chatroom.h b/Engine/Plugins/Online/OnlineFramework/Source/Party/Public/Chatroom.h
--- a/Engine/Plugins/Online/OnlineFramework/Source/Party/Public/Chatroom.h
+++ b/Engine/Plugins/Online/OnlineFramework/Source/Party/Public/Chatroom.h
@@ -53,8 +53,9 @@ public:
 	 * @param ChatRoomId chat room id to connect with
 	 * @param CompletionDelegate delegate fired when operation completes
 	 * @param Password optional password for the room
+     * NOTE: All parameters are by value because ClearTimer/SetTimer can destroy the internal lambda function
 	 */
-	void CreateOrJoinChatRoom(const FUniqueNetIdRepl& LocalUserId, const FChatRoomId& ChatRoomId, const FOnChatRoomCreatedOrJoined& CompletionDelegate, const FString& Password = FString());
+	void CreateOrJoinChatRoom(FUniqueNetIdRepl LocalUserId, FChatRoomId ChatRoomId, FOnChatRoomCreatedOrJoined CompletionDelegate, FString Password = FString());
 
 	/**
 	 * Leave the joined chat room
@@ -90,6 +91,14 @@ private:
 	 */
 	void OnChatRoomLeft(const FUniqueNetId& LocalUserId, const FChatRoomId& RoomId, bool bWasSuccessful, const FString& Error, FChatRoomId ChatRoomIdCopy, FOnChatRoomLeft CompletionDelegate);
 
+	/**
+	 * Common code called at the end of chat room cleanup
+	 *
+	 * @param RoomId id of the room let
+	 * @param CompletionDelegate delegate to fire to notify that the room has been left
+	 */
+	void ChatRoomLeftInternal(const FChatRoomId& RoomId, const FOnChatRoomLeft& CompletionDelegate);
+
 private:
 
 	/** Current chat room associated with this object (FString so UPROPERTY works) */
diff --git a/Engine/Plugins/Online/OnlineFramework/Source/Qos/Private/QosEvaluator.cpp b/Engine/Plugins/Online/OnlineFramework/Source/Qos/Private/QosEvaluator.cpp
--- a/Engine/Plugins/Online/OnlineFramework/Source/Qos/Private/QosEvaluator.cpp
+++ b/Engine/Plugins/Online/OnlineFramework/Source/Qos/Private/QosEvaluator.cpp
@@ -104,6 +104,10 @@ void UQosEvaluator::FindDatacenters(const FQosParams& InParams, TArray<FQosDatac
 			{
 				new (Datacenters) FQosRegionInfo(Datacenter);
 			}
+			else
+			{
+				UE_LOG(LogQos, Verbose, TEXT("Skipping region [%s]"), *Datacenter.RegionId);
+			}
 		}
 
 		if (!InParams.bUseOldQosServers)
@@ -173,7 +177,8 @@ void UQosEvaluator::PingRegionServers(const FQosParams& InParams, const FOnQosSe
 			const FString& RegionId = Region.Region.RegionId;
 			const int32 NumServers = Region.Region.Servers.Num();
 			int32 ServerIdx = FMath::RandHelper(NumServers);
-			Region.Result = NumServers > 0 ? EQosCompletionResult::Invalid : EQosCompletionResult::Success;
+			// Default to invalid ping tests and set it to something else later
+			Region.Result = EQosCompletionResult::Invalid;
 			if (NumServers > 0)
 			{
 				for (int32 PingIdx = 0; PingIdx < NumTestsPerRegion; PingIdx++)
@@ -181,16 +186,18 @@ void UQosEvaluator::PingRegionServers(const FQosParams& InParams, const FOnQosSe
 					const FQosPingServerInfo& Server = Region.Region.Servers[ServerIdx];
 					const FString Address = FString::Printf(TEXT("%s:%d"), *Server.Address, Server.Port);
 
-					auto CompletionDelegate = [WeakThisCap, RegionId, NumTestsPerRegion, InCompletionDelegate](FIcmpEchoResult Result)
+					auto CompletionDelegate = [WeakThisCap, RegionId, NumTestsPerRegion, InCompletionDelegate](FIcmpEchoResult InResult)
 					{
 						if (WeakThisCap.IsValid())
 						{
 							auto StrongThis = WeakThisCap.Get();
-							StrongThis->OnPingResultComplete(RegionId, NumTestsPerRegion, Result);
+							StrongThis->OnPingResultComplete(RegionId, NumTestsPerRegion, InResult);
 							if (StrongThis->AreAllRegionsComplete())
 							{
+								EQosCompletionResult Result = EQosCompletionResult::Success;
 								StrongThis->CalculatePingAverages();
-								InCompletionDelegate.ExecuteIfBound(EQosCompletionResult::Success, StrongThis->Datacenters);
+								StrongThis->EndAnalytics(Result);
+								InCompletionDelegate.ExecuteIfBound(Result, StrongThis->Datacenters);
 								StrongThis->bInProgress = false;
 							}
 						}
@@ -202,6 +209,14 @@ void UQosEvaluator::PingRegionServers(const FQosParams& InParams, const FOnQosSe
 					bDidNothing = false;
 				}
 			}
+			else
+			{
+				UE_LOG(LogQos, Verbose, TEXT("Nothing to ping [%s]"), *RegionId);
+			}
+		}
+		else
+		{
+			UE_LOG(LogQos, Verbose, TEXT("Region disabled [%s]"), *Region.Region.RegionId);
 		}
 	}
 
@@ -567,7 +582,7 @@ void UQosEvaluator::OnPingResultComplete(const FString& RegionId, int32 NumTests
 
 			if (QosStats.IsValid())
 			{
-				QosStats->RecordQosAttempt(RegionId, PingInMs, bSuccess);
+				QosStats->RecordQosAttempt(RegionId, Result.ResolvedAddress, PingInMs, bSuccess);
 			}
 
 			if (Region.PingResults.Num() == NumTests)
diff --git a/Engine/Plugins/Online/OnlineFramework/Source/Qos/Private/QosRegionManager.cpp b/Engine/Plugins/Online/OnlineFramework/Source/Qos/Private/QosRegionManager.cpp
--- a/Engine/Plugins/Online/OnlineFramework/Source/Qos/Private/QosRegionManager.cpp
+++ b/Engine/Plugins/Online/OnlineFramework/Source/Qos/Private/QosRegionManager.cpp
@@ -42,16 +42,6 @@ UQosRegionManager::UQosRegionManager(const FObjectInitializer& ObjectInitializer
 
 	// get a forced region id from the command line as an override
 	FParse::Value(FCommandLine::Get(), TEXT("McpRegion="), ForceRegionId);
-
-	// Temporary hack to aid in config conversion. 
-	if (ForceRegionId == TEXT("USA"))
-	{
-		ForceRegionId = TEXT("NA");
-	}
-	else if (ForceRegionId == TEXT("Poland"))
-	{
-		ForceRegionId = TEXT("EU");
-	}
 }
 
 void UQosRegionManager::PostReloadConfig(UProperty* PropertyThatWasLoaded)
@@ -219,10 +209,12 @@ void UQosRegionManager::OnQosEvaluationComplete(EQosCompletionResult Result, con
 		QosEvalResult = EQosCompletionResult::Failure;
 	}
 
-	if (QosEvalResult == EQosCompletionResult::Success)
+	if (QosEvalResult == EQosCompletionResult::Success ||
+		QosEvalResult == EQosCompletionResult::Failure)
 	{
 		if (RegionOptions.Num() > 0)
 		{
+			// Try to set something regardless of Qos result
 			TrySetDefaultRegion();
 		}
 	}
@@ -248,7 +240,14 @@ FString UQosRegionManager::GetRegionId() const
 	{
 		// if we haven't run the evaluator just use the region from settings
 		// development dedicated server will come here, live services should use -mcpregion
-		return TEXT("NONE");
+		return NO_REGION;
+	}
+
+	if (SelectedRegionId.IsEmpty())
+	{
+		// Always set some kind of region, empty implies "wildcard" to the matchmaking code
+		UE_LOG(LogQos, Verbose, TEXT("No region found, returning REGION_NONE"));
+		return NO_REGION;
 	}
 
 	return SelectedRegionId;
@@ -294,7 +293,8 @@ void UQosRegionManager::TrySetDefaultRegion()
 			FString BestRegionId;
 			for (const FQosRegionInfo& Region : RegionOptions)
 			{
-				if (!Region.Region.bBeta && Region.AvgPingMs < BestPing)
+				if (Region.IsUsable() && (Region.Result == EQosCompletionResult::Success) &&
+					!Region.Region.bBeta && Region.AvgPingMs < BestPing)
 				{
 					BestPing = Region.AvgPingMs;
 					BestRegionId = Region.Region.RegionId;
diff --git a/Engine/Plugins/Online/OnlineFramework/Source/Qos/Private/QosStats.cpp b/Engine/Plugins/Online/OnlineFramework/Source/Qos/Private/QosStats.cpp
--- a/Engine/Plugins/Online/OnlineFramework/Source/Qos/Private/QosStats.cpp
+++ b/Engine/Plugins/Online/OnlineFramework/Source/Qos/Private/QosStats.cpp
@@ -96,7 +96,7 @@ void FQosDatacenterStats::RecordQosAttempt(const FOnlineSessionSearchResult& Sea
 		QosData.NumSuccessAttempts += bSuccess ? 1 : 0;
 
 		FQosStats_QosSearchResult& NewSearchResult = *new (QosData.SearchResults) FQosStats_QosSearchResult();
-		NewSearchResult.OwnerId = SearchResult.Session.OwningUserId;
+		NewSearchResult.OwnerId = SearchResult.Session.OwningUserId.IsValid() ? SearchResult.Session.OwningUserId->ToString() : TEXT("Unknown");
 		NewSearchResult.PingInMs = SearchResult.PingInMs;
 		NewSearchResult.bIsValid = SearchResult.IsValid();
 
@@ -108,7 +108,7 @@ void FQosDatacenterStats::RecordQosAttempt(const FOnlineSessionSearchResult& Sea
 	}
 }
 
-void FQosDatacenterStats::RecordQosAttempt(const FString& Region, int32 PingInMs, bool bSuccess)
+void FQosDatacenterStats::RecordQosAttempt(const FString& Region, const FString& OwnerId, int32 PingInMs, bool bSuccess)
 {
 	if (bAnalyticsInProgress)
 	{
@@ -116,7 +116,7 @@ void FQosDatacenterStats::RecordQosAttempt(const FString& Region, int32 PingInMs
 		QosData.NumSuccessAttempts += bSuccess ? 1 : 0;
 
 		FQosStats_QosSearchResult& NewSearchResult = *new (QosData.SearchResults) FQosStats_QosSearchResult();
-		NewSearchResult.OwnerId = nullptr;
+		NewSearchResult.OwnerId = OwnerId;
 		NewSearchResult.PingInMs = PingInMs;
 		NewSearchResult.DatacenterId = Region;
 		NewSearchResult.bIsValid = true;
@@ -205,7 +205,7 @@ void FQosDatacenterStats::ParseQosResults(TSharedPtr<IAnalyticsProvider>& Analyt
 		for (auto& SearchResult : QosData.SearchResults)
 		{
 			SearchDetails += FString::Printf(TEXT("{\"OwnerId\":\"%s\", \"RegionId\":\"%s\", \"PingInMs\":%d, \"bIsValid\":%s},"),
-				SearchResult.OwnerId.IsValid() ? *SearchResult.OwnerId->ToString() : TEXT("Unknown"),
+				*SearchResult.OwnerId,
 				*SearchResult.DatacenterId,
 				SearchResult.PingInMs,
 				SearchResult.bIsValid ? TEXT("true") : TEXT("false")
diff --git a/Engine/Plugins/Online/OnlineFramework/Source/Qos/Private/QosStats.h b/Engine/Plugins/Online/OnlineFramework/Source/Qos/Private/QosStats.h
--- a/Engine/Plugins/Online/OnlineFramework/Source/Qos/Private/QosStats.h
+++ b/Engine/Plugins/Online/OnlineFramework/Source/Qos/Private/QosStats.h
@@ -92,7 +92,7 @@ private:
 	struct FQosStats_QosSearchResult
 	{
 		/** Owner of the session */
-		TSharedPtr<const FUniqueNetId> OwnerId;
+		FString OwnerId;
 		/** Datacenter Id */
 		FString DatacenterId;
 		/** Ping time */
@@ -101,7 +101,6 @@ private:
 		bool bIsValid;
 
 		FQosStats_QosSearchResult() :
-			OwnerId(NULL),
 			PingInMs(0),
 			bIsValid(false)
 		{}
@@ -218,10 +217,11 @@ public:
 	 * Record a single ping attempt
 	 * 
 	 * @param Region region of the server
+	 * @param OwnerId the owner of the server (ip address)
 	 * @param PingInMs ping to the server
 	 * @param bSuccess was the attempt successful
 	 */
-	void RecordQosAttempt(const FString& Region, int32 PingInMs, bool bSuccess);
+	void RecordQosAttempt(const FString& Region, const FString& OwnerId, int32 PingInMs, bool bSuccess);
 
 	/**
 	 * End recording of a Qos determination
diff --git a/Engine/Plugins/Online/OnlineFramework/Source/Qos/Public/QosInterface.h b/Engine/Plugins/Online/OnlineFramework/Source/Qos/Public/QosInterface.h
--- a/Engine/Plugins/Online/OnlineFramework/Source/Qos/Public/QosInterface.h
+++ b/Engine/Plugins/Online/OnlineFramework/Source/Qos/Public/QosInterface.h
@@ -8,6 +8,8 @@
 
 class IAnalyticsProvider;
 
+#define NO_REGION TEXT("NONE")
+
 /**
  * Main Qos interface for actions related to server quality of service
  */
diff --git a/Engine/Plugins/Online/OnlineFramework/Source/Rejoin/Private/RejoinModule.cpp b/Engine/Plugins/Online/OnlineFramework/Source/Rejoin/Private/RejoinModule.cpp
new file mode 100644
--- /dev/null
+++ b/Engine/Plugins/Online/OnlineFramework/Source/Rejoin/Private/RejoinModule.cpp
@@ -0,0 +1,13 @@
+// Copyright 1998-2016 Epic Games, Inc. All Rights Reserved.
+
+#include "RejoinModule.h"
+
+IMPLEMENT_MODULE(FRejoinModule, Rejoin);
+
+void FRejoinModule::StartupModule()
+{
+}
+
+void FRejoinModule::ShutdownModule()
+{
+}
diff --git a/Engine/Plugins/Online/OnlineFramework/Source/Rejoin/Public/RejoinModule.h b/Engine/Plugins/Online/OnlineFramework/Source/Rejoin/Public/RejoinModule.h
new file mode 100644
--- /dev/null
+++ b/Engine/Plugins/Online/OnlineFramework/Source/Rejoin/Public/RejoinModule.h
@@ -0,0 +1,17 @@
+// Copyright 1998-2016 Epic Games, Inc. All Rights Reserved.
+
+#pragma once
+
+#include "ModuleManager.h"
+
+/**
+ * Module for Rejoin base implementation
+ */
+class FRejoinModule : 
+	public IModuleInterface
+{
+private:
+// IModuleInterface
+	virtual void StartupModule() override;
+	virtual void ShutdownModule() override;
+};
diff --git a/Engine/Plugins/Online/OnlineSubsystem/Source/OnlineSubsystem.Build.cs b/Engine/Plugins/Online/OnlineSubsystem/Source/OnlineSubsystem.Build.cs
--- a/Engine/Plugins/Online/OnlineSubsystem/Source/OnlineSubsystem.Build.cs
+++ b/Engine/Plugins/Online/OnlineSubsystem/Source/OnlineSubsystem.Build.cs
@@ -6,13 +6,15 @@ public class OnlineSubsystem : ModuleRules
 {
 	public OnlineSubsystem(TargetInfo Target)
 	{
+		PCHUsage = ModuleRules.PCHUsageMode.UseExplicitOrSharedPCHs;
+		
 		PublicDependencyModuleNames.AddRange(
 			new string[] {
 				"Json",
 			}
 		);
 
-        PublicIncludePaths.Add("OnlineSubsystem/Public/Interfaces");
+		PublicIncludePaths.Add("OnlineSubsystem/Public/Interfaces");
 
         PrivateIncludePaths.Add("OnlineSubsystem/Private");
 
diff --git a/Engine/Plugins/Online/OnlineSubsystem/Source/Private/OnlineSubsystemImpl.cpp b/Engine/Plugins/Online/OnlineSubsystem/Source/Private/OnlineSubsystemImpl.cpp
--- a/Engine/Plugins/Online/OnlineSubsystem/Source/Private/OnlineSubsystemImpl.cpp
+++ b/Engine/Plugins/Online/OnlineSubsystem/Source/Private/OnlineSubsystemImpl.cpp
@@ -7,6 +7,7 @@
 #include "Interfaces/OnlineIdentityInterface.h"
 #include "Interfaces/OnlineSessionInterface.h"
 #include "Interfaces/OnlineFriendsInterface.h"
+#include "Interfaces/OnlinePurchaseInterface.h"
 
 namespace OSSConsoleVariables
 {
@@ -42,6 +43,10 @@ FOnlineSubsystemImpl::~FOnlineSubsystemImpl()
 	ensure(!TickHandle.IsValid());
 }
 
+void FOnlineSubsystemImpl::PreUnload()
+{
+}
+
 bool FOnlineSubsystemImpl::Shutdown()
 {
 	OnNamedInterfaceCleanup();
@@ -76,6 +81,7 @@ void FOnlineSubsystemImpl::StopTicker()
 
 bool FOnlineSubsystemImpl::Tick(float DeltaTime)
 {
+	QUICK_SCOPE_CYCLE_COUNTER(STAT_FOnlineSubsystemImpl_Tick);
 	if (!NextTickQueue.IsEmpty())
 	{
 		// unload the next-tick queue into our buffer. Any further executes (from within callbacks) will happen NEXT frame (as intended)
@@ -88,6 +94,7 @@ bool FOnlineSubsystemImpl::Tick(float DeltaTime)
 		// execute any functions in the current tick array
 		for (const auto& Callback : CurrentTickBuffer)
 		{
+			QUICK_SCOPE_CYCLE_COUNTER(STAT_FOnlineSubsystemImpl_Tick_ExecuteCallback);
 			Callback.ExecuteIfBound();
 		}
 		CurrentTickBuffer.SetNum(0, false); // keep the memory around
@@ -199,6 +206,99 @@ bool FOnlineSubsystemImpl::Exec(UWorld* InWorld, const TCHAR* Cmd, FOutputDevice
 	{
 		bWasHandled = HandleSessionExecCommands(InWorld, Cmd, Ar);
 	}
+	else if (FParse::Command(&Cmd, TEXT("PURCHASE")))
+	{
+		bWasHandled = HandlePurchaseExecCommands(InWorld, Cmd, Ar);
+	}
+	
+	return bWasHandled;
+}
+
+void FOnlineSubsystemImpl::DumpReceipts(const FUniqueNetId& UserId)
+{
+	IOnlinePurchasePtr PurchaseInt = GetPurchaseInterface();
+	if (PurchaseInt.IsValid())
+	{
+		TArray<FPurchaseReceipt> Receipts;
+		PurchaseInt->GetReceipts(UserId, Receipts);
+		for (const FPurchaseReceipt& Receipt : Receipts)
+		{
+			UE_LOG_ONLINE(Display, TEXT("Receipt: %s %d"),
+						  *Receipt.TransactionId,
+						  (int32)Receipt.TransactionState);
+			
+			UE_LOG_ONLINE(Display, TEXT("-Offers:"));
+			for (const FPurchaseReceipt::FReceiptOfferEntry& ReceiptOffer : Receipt.ReceiptOffers)
+			{
+				UE_LOG_ONLINE(Display, TEXT(" -Namespace: %s Id: %s Quantity: %d"),
+							  *ReceiptOffer.Namespace,
+							  *ReceiptOffer.OfferId,
+							  ReceiptOffer.Quantity);
+				
+				UE_LOG_ONLINE(Display, TEXT(" -LineItems:"));
+				for (const FPurchaseReceipt::FLineItemInfo& LineItem : ReceiptOffer.LineItems)
+				{
+					UE_LOG_ONLINE(Display, TEXT("  -Name: %s Id: %s ValidationInfo: %d bytes"),
+								  *LineItem.ItemName,
+								  *LineItem.UniqueId,
+								  LineItem.ValidationInfo.Len());
+				}
+			}
+		}
+	}
+}
+
+void FOnlineSubsystemImpl::OnQueryReceiptsComplete(const FOnlineError& Result, TSharedPtr<const FUniqueNetId> UserId)
+{
+	UE_LOG_ONLINE(Display, TEXT("OnQueryReceiptsComplete %s"), Result.ToLogString());
+	DumpReceipts(*UserId);
+}
+
+bool FOnlineSubsystemImpl::HandlePurchaseExecCommands(UWorld* InWorld, const TCHAR* Cmd, FOutputDevice& Ar)
+{
+	bool bWasHandled = false;
+	
+	if (FParse::Command(&Cmd, TEXT("RECEIPTS")))
+	{
+		IOnlinePurchasePtr PurchaseInt = GetPurchaseInterface();
+		IOnlineIdentityPtr IdentityInt = GetIdentityInterface();
+		if (PurchaseInt.IsValid() && IdentityInt.IsValid())
+		{
+			FString CommandStr = FParse::Token(Cmd, false);
+			if (CommandStr.IsEmpty())
+			{
+				UE_LOG_ONLINE(Warning, TEXT("usage: PURCHASE RECEIPTS <command> <userid>"));
+			}
+			else
+			{
+				FString UserIdStr = FParse::Token(Cmd, false);
+				if (UserIdStr.IsEmpty())
+				{
+					UE_LOG_ONLINE(Warning, TEXT("usage: PURCHASE RECEIPTS <command> <userid>"));
+				}
+				else
+				{
+					TSharedPtr<const FUniqueNetId> UserId = IdentityInt->CreateUniquePlayerId(UserIdStr);
+					if (UserId.IsValid())
+					{
+						if (CommandStr == TEXT("RESTORE"))
+						{
+							FOnQueryReceiptsComplete CompletionDelegate;
+							CompletionDelegate.BindRaw(this, &FOnlineSubsystemImpl::OnQueryReceiptsComplete, UserId);
+							PurchaseInt->QueryReceipts(*UserId, true, CompletionDelegate);
+							
+							bWasHandled = true;
+						}
+						else if (CommandStr == TEXT("DUMP"))
+						{
+							DumpReceipts(*UserId);
+							bWasHandled = true;
+						}
+					}
+				}
+			}
+		}
+	}
 
 	return bWasHandled;
 }
diff --git a/Engine/Plugins/Online/OnlineSubsystem/Source/Private/OnlineSubsystemModule.cpp b/Engine/Plugins/Online/OnlineSubsystem/Source/Private/OnlineSubsystemModule.cpp
--- a/Engine/Plugins/Online/OnlineSubsystem/Source/Private/OnlineSubsystemModule.cpp
+++ b/Engine/Plugins/Online/OnlineSubsystem/Source/Private/OnlineSubsystemModule.cpp
@@ -59,11 +59,22 @@ static TSharedPtr<IModuleInterface> LoadSubsystemModule(const FString& Subsystem
 
 void FOnlineSubsystemModule::StartupModule()
 {
+	// These should not be LoadModuleChecked because these modules might not exist
+	// Load dependent modules to ensure they will still exist during ShutdownModule.
+	// We will alwawys load these modules at the cost of extra modules loaded for the few OSS (like Null) that don't use it.
+	FModuleManager::Get().LoadModule(TEXT("HTTP"));
+	FModuleManager::Get().LoadModule(TEXT("XMPP"));
+
 	LoadDefaultSubsystem();
 	// Also load the console/platform specific OSS which might not necessarily be the default OSS instance
 	IOnlineSubsystem::GetByPlatform();
 }
 
+void FOnlineSubsystemModule::PreUnloadCallback()
+{
+	PreUnloadOnlineSubsystem();
+}
+
 void FOnlineSubsystemModule::ShutdownModule()
 {
 	ShutdownOnlineSubsystem();
@@ -122,6 +133,15 @@ void FOnlineSubsystemModule::ReloadDefaultSubsystem()
 	LoadDefaultSubsystem();
 }
 
+void FOnlineSubsystemModule::PreUnloadOnlineSubsystem()
+{
+	// Shutdown all online subsystem instances
+	for (TMap<FName, IOnlineSubsystemPtr>::TIterator It(OnlineSubsystems); It; ++It)
+	{
+		It.Value()->PreUnload();
+	}
+}
+
 void FOnlineSubsystemModule::ShutdownOnlineSubsystem()
 {
 	FModuleManager& ModuleManager = FModuleManager::Get();
diff --git a/Engine/Plugins/Online/OnlineSubsystem/Source/Public/Interfaces/OnlineExternalUIInterface.h b/Engine/Plugins/Online/OnlineSubsystem/Source/Public/Interfaces/OnlineExternalUIInterface.h
--- a/Engine/Plugins/Online/OnlineSubsystem/Source/Public/Interfaces/OnlineExternalUIInterface.h
+++ b/Engine/Plugins/Online/OnlineSubsystem/Source/Public/Interfaces/OnlineExternalUIInterface.h
@@ -187,10 +187,11 @@ public:
 	 *	Displays the UI that shows a user's list of friends to invite
 	 *
 	 * @param LocalUserNum the controller number of the associated user
+	 * @param SessionName name of session associated with the invite
 	 *
 	 * @return true if it was able to show the UI, false if it failed
 	 */
-	virtual bool ShowInviteUI(int32 LocalUserNum, FName SessionMame = GameSessionName) = 0;
+	virtual bool ShowInviteUI(int32 LocalUserNum, FName SessionName = GameSessionName) = 0;
 
 	/**
 	 *	Displays the UI that shows a user's list of achievements
diff --git a/Engine/Plugins/Online/OnlineSubsystem/Source/Public/Interfaces/OnlinePresenceInterface.h b/Engine/Plugins/Online/OnlineSubsystem/Source/Public/Interfaces/OnlinePresenceInterface.h
--- a/Engine/Plugins/Online/OnlineSubsystem/Source/Public/Interfaces/OnlinePresenceInterface.h
+++ b/Engine/Plugins/Online/OnlineSubsystem/Source/Public/Interfaces/OnlinePresenceInterface.h
@@ -156,6 +156,15 @@ public:
 DECLARE_MULTICAST_DELEGATE_TwoParams(FOnPresenceReceived, const class FUniqueNetId& /*UserId*/, const TSharedRef<FOnlineUserPresence>& /*Presence*/);
 typedef FOnPresenceReceived::FDelegate FOnPresenceReceivedDelegate;
 
+/**
+ * Delegate executed when the array of presence data for a user changes.
+ *
+ * @param UserId The unique id of the user whose presence was received.
+ * @param PresenceArray The updated presence array
+ */
+DECLARE_MULTICAST_DELEGATE_TwoParams(FOnPresenceArrayUpdated, const class FUniqueNetId& /*UserId*/, const TArray<TSharedRef<FOnlineUserPresence> >& /*PresenceArray*/);
+typedef FOnPresenceArrayUpdated::FDelegate FOnPresenceArrayUpdatedDelegate;
+
 /**
  *	Interface class for getting and setting rich presence information.
  */
@@ -200,6 +209,14 @@ public:
 	 */
 	DEFINE_ONLINE_DELEGATE_TWO_PARAM(OnPresenceReceived, const class FUniqueNetId& /*UserId*/, const TSharedRef<FOnlineUserPresence>& /*Presence*/);
 
+	/**
+ 	 * Delegate executed when the array of presence data for a user changes.
+ 	 *
+	 * @param UserId The unique id of the user whose presence was received.
+ 	 * @param PresenceArray The updated presence array
+ 	 */
+	DEFINE_ONLINE_DELEGATE_TWO_PARAM(OnPresenceArrayUpdated, const class FUniqueNetId& /*UserId*/, const TArray<TSharedRef<FOnlineUserPresence> >& /*NewPresenceArray*/);
+
 	/**
 	 * Gets the cached presence information for a user.
 	 *
diff --git a/Engine/Plugins/Online/OnlineSubsystem/Source/Public/Interfaces/OnlinePurchaseInterface.h b/Engine/Plugins/Online/OnlineSubsystem/Source/Public/Interfaces/OnlinePurchaseInterface.h
--- a/Engine/Plugins/Online/OnlineSubsystem/Source/Public/Interfaces/OnlinePurchaseInterface.h
+++ b/Engine/Plugins/Online/OnlineSubsystem/Source/Public/Interfaces/OnlinePurchaseInterface.h
@@ -20,24 +20,28 @@ public:
 	 * @param InNamespace namespace of offer to be purchased
 	 * @param InOfferId id of offer to be purchased
 	 * @param InQuantity number to purchase
+	 * @param bInIsConsumable is the offer consumable or one time purchase
 	 */
-	void AddPurchaseOffer(const FOfferNamespace& InNamespace, const FUniqueOfferId& InOfferId, int32 InQuantity)
+	void AddPurchaseOffer(const FOfferNamespace& InNamespace, const FUniqueOfferId& InOfferId, int32 InQuantity, bool bInIsConsumable = true)
 	{ 
-		PurchaseOffers.Add(FPurchaseOfferEntry(InNamespace, InOfferId, InQuantity)); 
+		PurchaseOffers.Add(FPurchaseOfferEntry(InNamespace, InOfferId, InQuantity, bInIsConsumable));
 	}
 	/**
 	 * Single offer entry for purchase 
 	 */
 	struct FPurchaseOfferEntry
 	{
-		FPurchaseOfferEntry(const FOfferNamespace& InOfferNamespace, const FUniqueOfferId& InOfferId, int32 InQuantity)
+		FPurchaseOfferEntry(const FOfferNamespace& InOfferNamespace, const FUniqueOfferId& InOfferId, int32 InQuantity, bool bInIsConsumable)
 			: OfferNamespace(InOfferNamespace)
 			, OfferId(InOfferId)
 			, Quantity(InQuantity)
 		{ }
 
+		/** Namespace in which the offer resides */
 		FOfferNamespace OfferNamespace;
+		/** Platform specific offer id (defined on backend) */
 		FUniqueOfferId OfferId;
+		/** Number of offers of this type to purchase */
 		int32 Quantity;
 	};	
 	/** List of offers being purchased */
@@ -47,24 +51,27 @@ public:
 /**
  * State of a purchase transaction
  */
-namespace EPurchaseTransactionState
+enum class EPurchaseTransactionState : uint8
 {
-	enum Type
-	{
-		/** processing has not started on the purchase */
-		NotStarted,
-		/** currently processing the purchase */
-		Processing,
-		/** purchase completed successfully */
-		Purchased,
-		/** purchase completed but failed */
-		Failed,
-		/** purchase canceled by user */
-		Canceled,
-		/** prior purchase that has been restored */
-		Restored
-	};
-}
+	/** processing has not started on the purchase */
+	NotStarted,
+	/** currently processing the purchase */
+	Processing,
+	/** purchase completed successfully */
+	Purchased,
+	/** purchase completed but failed */
+	Failed,
+	/** purchase has been deferred (neither failed nor completed) */
+	Deferred,
+	/** purchase canceled by user */
+	Canceled,
+	/** prior purchase that has been restored */
+	Restored,
+	/** purchase failed as not allowed */
+	NotAllowed,
+	/** purchase failed as invalid */
+	Invalid
+};
 
 /**
  * Receipt result from checkout
@@ -131,7 +138,7 @@ public:
 	/** Unique Id for this transaction/order */
 	FString TransactionId;
 	/** Current state of the purchase */
-	EPurchaseTransactionState::Type TransactionState;
+	EPurchaseTransactionState TransactionState;
 
 	/** List of offers that were purchased */
 	TArray<FReceiptOfferEntry> ReceiptOffers;
@@ -187,7 +194,16 @@ public:
 	 * @param Delegate completion callback (guaranteed to be called)
 	 */
 	virtual void Checkout(const FUniqueNetId& UserId, const FPurchaseCheckoutRequest& CheckoutRequest, const FOnPurchaseCheckoutComplete& Delegate) = 0;
-
+	
+	/**
+	 * Finalizes a purchase with the supporting platform
+	 * Acknowledges that the purchase has been properly redeemed by the application
+	 *
+	 * @param UserId user where the purchase was made
+	 * @param ReceiptId purchase id for this platform
+	 */
+	virtual void FinalizePurchase(const FUniqueNetId& UserId, const FString& ReceiptId) = 0;
+	
 	/**
 	 * Initiate the checkout process for obtaining offers via code redemption
 	 *
@@ -201,9 +217,10 @@ public:
 	 * Query for all of the user's receipts from prior purchases
 	 *
 	 * @param UserId user initiating the request
+	 * @param bRestoreReceipts initiate recovery of any receipts on the specific platform
 	 * @param Delegate completion callback (guaranteed to be called)
 	 */
-	virtual void QueryReceipts(const FUniqueNetId& UserId, const FOnQueryReceiptsComplete& Delegate) = 0;
+	virtual void QueryReceipts(const FUniqueNetId& UserId, bool bRestoreReceipts, const FOnQueryReceiptsComplete& Delegate) = 0;
 
 	/**
 	 * Get list of cached receipts for user (includes transactions currently being processed)
