Fixup headers and remove exception for OnlineSubsystemIOS including non-IWYU headers.

#rnx

diff --git a/Engine/Plugins/Online/IOS/OnlineSubsystemIOS/Source/Private/OnlineAchievementsInterfaceIOS.cpp b/Engine/Plugins/Online/IOS/OnlineSubsystemIOS/Source/Private/OnlineAchievementsInterfaceIOS.cpp
--- a/Engine/Plugins/Online/IOS/OnlineSubsystemIOS/Source/Private/OnlineAchievementsInterfaceIOS.cpp
+++ b/Engine/Plugins/Online/IOS/OnlineSubsystemIOS/Source/Private/OnlineAchievementsInterfaceIOS.cpp
@@ -1,7 +1,8 @@
 // Copyright 1998-2019 Epic Games, Inc. All Rights Reserved.
 
 // Module includes
-#include "OnlineSubsystemIOSPrivatePCH.h"
+#include "OnlineAchievementsInterfaceIOS.h"
+#include "OnlineSubsystem.h"
 
 // Game Center includes
 #include <GameKit/GKAchievement.h>
diff --git a/Engine/Plugins/Online/IOS/OnlineSubsystemIOS/Source/Private/OnlineAppStoreUtils.cpp b/Engine/Plugins/Online/IOS/OnlineSubsystemIOS/Source/Private/OnlineAppStoreUtils.cpp
--- a/Engine/Plugins/Online/IOS/OnlineSubsystemIOS/Source/Private/OnlineAppStoreUtils.cpp
+++ b/Engine/Plugins/Online/IOS/OnlineSubsystemIOS/Source/Private/OnlineAppStoreUtils.cpp
@@ -1,7 +1,7 @@
 // Copyright 1998-2019 Epic Games, Inc. All Rights Reserved.
 
-#include "OnlineSubsystemIOSPrivatePCH.h"
-#import "OnlineAppStoreUtils.h"
+#include "OnlineAppStoreUtils.h"
+#include "OnlineSubsystem.h"
 
 #include "HAL/PlatformProcess.h"
 
diff --git a/Engine/Plugins/Online/IOS/OnlineSubsystemIOS/Source/Private/OnlineAppStoreUtils.h b/Engine/Plugins/Online/IOS/OnlineSubsystemIOS/Source/Private/OnlineAppStoreUtils.h
--- a/Engine/Plugins/Online/IOS/OnlineSubsystemIOS/Source/Private/OnlineAppStoreUtils.h
+++ b/Engine/Plugins/Online/IOS/OnlineSubsystemIOS/Source/Private/OnlineAppStoreUtils.h
@@ -1,5 +1,7 @@
 // Copyright 1998-2019 Epic Games, Inc. All Rights Reserved.
 
+#pragma once
+
 #include "OnlineSubsystemTypes.h"
 
 /**
diff --git a/Engine/Plugins/Online/IOS/OnlineSubsystemIOS/Source/Private/OnlineExternalUIInterfaceIOS.cpp b/Engine/Plugins/Online/IOS/OnlineSubsystemIOS/Source/Private/OnlineExternalUIInterfaceIOS.cpp
--- a/Engine/Plugins/Online/IOS/OnlineSubsystemIOS/Source/Private/OnlineExternalUIInterfaceIOS.cpp
+++ b/Engine/Plugins/Online/IOS/OnlineSubsystemIOS/Source/Private/OnlineExternalUIInterfaceIOS.cpp
@@ -1,8 +1,9 @@
 // Copyright 1998-2019 Epic Games, Inc. All Rights Reserved.
 
 // Module includes
-#include "OnlineSubsystemIOSPrivatePCH.h"
+#include "OnlineExternalUIInterfaceIOS.h"
 #include "OnlineIdentityInterfaceIOS.h"
+#include "OnlineSubsystemIOS.h"
 
 FOnlineExternalUIIOS::FOnlineExternalUIIOS(FOnlineSubsystemIOS* InSubsystem)
 	: Subsystem(InSubsystem)
diff --git a/Engine/Plugins/Online/IOS/OnlineSubsystemIOS/Source/Private/OnlineFriendsInterfaceIOS.cpp b/Engine/Plugins/Online/IOS/OnlineSubsystemIOS/Source/Private/OnlineFriendsInterfaceIOS.cpp
--- a/Engine/Plugins/Online/IOS/OnlineSubsystemIOS/Source/Private/OnlineFriendsInterfaceIOS.cpp
+++ b/Engine/Plugins/Online/IOS/OnlineSubsystemIOS/Source/Private/OnlineFriendsInterfaceIOS.cpp
@@ -1,7 +1,7 @@
 // Copyright 1998-2019 Epic Games, Inc. All Rights Reserved.
 
-
-#include "OnlineSubsystemIOSPrivatePCH.h"
+#include "OnlineFriendsInterfaceIOS.h"
+#include "OnlineSubsystemIOS.h"
 
 // FOnlineFriendIOS
 
diff --git a/Engine/Plugins/Online/IOS/OnlineSubsystemIOS/Source/Private/OnlineIdentityInterfaceIOS.cpp b/Engine/Plugins/Online/IOS/OnlineSubsystemIOS/Source/Private/OnlineIdentityInterfaceIOS.cpp
--- a/Engine/Plugins/Online/IOS/OnlineSubsystemIOS/Source/Private/OnlineIdentityInterfaceIOS.cpp
+++ b/Engine/Plugins/Online/IOS/OnlineSubsystemIOS/Source/Private/OnlineIdentityInterfaceIOS.cpp
@@ -1,7 +1,10 @@
 // Copyright 1998-2019 Epic Games, Inc. All Rights Reserved.
 
-#include "OnlineSubsystemIOSPrivatePCH.h"
+#include "OnlineIdentityInterfaceIOS.h"
+#include "OnlineSubsystem.h"
+#include "OnlineSubsystemIOS.h"
 #include "OnlineError.h"
+#include "IOS/IOSAppDelegate.h"
 #import "OnlineAppStoreUtils.h"
 
 FOnlineIdentityIOS::FOnlineIdentityIOS()
diff --git a/Engine/Plugins/Online/IOS/OnlineSubsystemIOS/Source/Private/OnlineIdentityInterfaceIOS.h b/Engine/Plugins/Online/IOS/OnlineSubsystemIOS/Source/Private/OnlineIdentityInterfaceIOS.h
--- a/Engine/Plugins/Online/IOS/OnlineSubsystemIOS/Source/Private/OnlineIdentityInterfaceIOS.h
+++ b/Engine/Plugins/Online/IOS/OnlineSubsystemIOS/Source/Private/OnlineIdentityInterfaceIOS.h
@@ -3,12 +3,12 @@
 
 #pragma once
 
+#import <GameKit/GKLocalPlayer.h>
 
 #include "OnlineSubsystemIOSTypes.h"
 #include "Interfaces/OnlineIdentityInterface.h"
 #include "OnlineSubsystemIOSPackage.h"
 
-
 class FOnlineIdentityIOS :
 	public IOnlineIdentity
 {
diff --git a/Engine/Plugins/Online/IOS/OnlineSubsystemIOS/Source/Private/OnlineLeaderboardsInterfaceIOS.cpp b/Engine/Plugins/Online/IOS/OnlineSubsystemIOS/Source/Private/OnlineLeaderboardsInterfaceIOS.cpp
--- a/Engine/Plugins/Online/IOS/OnlineSubsystemIOS/Source/Private/OnlineLeaderboardsInterfaceIOS.cpp
+++ b/Engine/Plugins/Online/IOS/OnlineSubsystemIOS/Source/Private/OnlineLeaderboardsInterfaceIOS.cpp
@@ -1,7 +1,8 @@
 // Copyright 1998-2019 Epic Games, Inc. All Rights Reserved.
 
 // Module includes
-#include "OnlineSubsystemIOSPrivatePCH.h"
+#include "OnlineLeaderboardsInterfaceIOS.h"
+#include "OnlineSubsystemIOS.h"
 
 // GameCenter includes
 #include <GameKit/GKLeaderboard.h>
diff --git a/Engine/Plugins/Online/IOS/OnlineSubsystemIOS/Source/Private/OnlinePurchaseIOS.cpp b/Engine/Plugins/Online/IOS/OnlineSubsystemIOS/Source/Private/OnlinePurchaseIOS.cpp
--- a/Engine/Plugins/Online/IOS/OnlineSubsystemIOS/Source/Private/OnlinePurchaseIOS.cpp
+++ b/Engine/Plugins/Online/IOS/OnlineSubsystemIOS/Source/Private/OnlinePurchaseIOS.cpp
@@ -1,8 +1,8 @@
 // Copyright 1998-2019 Epic Games, Inc. All Rights Reserved.
 
-#include "OnlineSubsystemIOSPrivatePCH.h"
 #include "OnlinePurchaseIOS.h"
 #include "OnlineError.h"
+#include "OnlineSubsystemIOS.h"
 #import "OnlineStoreKitHelper.h"
 
 /** Take successful transactions and route them through deferred pipeline */
diff --git a/Engine/Plugins/Online/IOS/OnlineSubsystemIOS/Source/Private/OnlineSessionInterfaceIOS.cpp b/Engine/Plugins/Online/IOS/OnlineSubsystemIOS/Source/Private/OnlineSessionInterfaceIOS.cpp
--- a/Engine/Plugins/Online/IOS/OnlineSubsystemIOS/Source/Private/OnlineSessionInterfaceIOS.cpp
+++ b/Engine/Plugins/Online/IOS/OnlineSubsystemIOS/Source/Private/OnlineSessionInterfaceIOS.cpp
@@ -1,7 +1,8 @@
 // Copyright 1998-2019 Epic Games, Inc. All Rights Reserved.
 
-#include "OnlineSubsystemIOSPrivatePCH.h"
+#include "OnlineSessionInterfaceIOS.h"
 #include "OnlineSessionSettings.h"
+#include "OnlineSubsystemIOS.h"
 
 #if !PLATFORM_TVOS // @todo tvos: What is up with all this being busted?? Multipeer is gone, but so is older stuff? What's to do??
 
diff --git a/Engine/Plugins/Online/IOS/OnlineSubsystemIOS/Source/Private/OnlineSessionInterfaceIOS.h b/Engine/Plugins/Online/IOS/OnlineSubsystemIOS/Source/Private/OnlineSessionInterfaceIOS.h
--- a/Engine/Plugins/Online/IOS/OnlineSubsystemIOS/Source/Private/OnlineSessionInterfaceIOS.h
+++ b/Engine/Plugins/Online/IOS/OnlineSubsystemIOS/Source/Private/OnlineSessionInterfaceIOS.h
@@ -1,7 +1,8 @@
 // Copyright 1998-2019 Epic Games, Inc. All Rights Reserved.
 
 #pragma once
-#include <GameKit/GKLocalPlayer.h>
+
+#import <GameKit/GKLocalPlayer.h>
 
 #include "Interfaces/OnlineSessionInterface.h"
 #include "OnlineSubsystemIOSTypes.h"
diff --git a/Engine/Plugins/Online/IOS/OnlineSubsystemIOS/Source/Private/OnlineSharedCloudInterfaceIOS.cpp b/Engine/Plugins/Online/IOS/OnlineSubsystemIOS/Source/Private/OnlineSharedCloudInterfaceIOS.cpp
--- a/Engine/Plugins/Online/IOS/OnlineSubsystemIOS/Source/Private/OnlineSharedCloudInterfaceIOS.cpp
+++ b/Engine/Plugins/Online/IOS/OnlineSubsystemIOS/Source/Private/OnlineSharedCloudInterfaceIOS.cpp
@@ -1,6 +1,5 @@
 // Copyright 1998-2019 Epic Games, Inc. All Rights Reserved.
 
-#include "OnlineSubsystemIOSPrivatePCH.h"
 #include "OnlineSharedCloudInterfaceIOS.h"
 
 //
diff --git a/Engine/Plugins/Online/IOS/OnlineSubsystemIOS/Source/Private/OnlineStoreIOS.cpp b/Engine/Plugins/Online/IOS/OnlineSubsystemIOS/Source/Private/OnlineStoreIOS.cpp
--- a/Engine/Plugins/Online/IOS/OnlineSubsystemIOS/Source/Private/OnlineStoreIOS.cpp
+++ b/Engine/Plugins/Online/IOS/OnlineSubsystemIOS/Source/Private/OnlineStoreIOS.cpp
@@ -1,10 +1,10 @@
 // Copyright 1998-2019 Epic Games, Inc. All Rights Reserved.
 
-#include "OnlineSubsystemIOSPrivatePCH.h"
 #include "OnlineStoreIOS.h"
 #include "Internationalization/Internationalization.h"
 #include "Internationalization/Culture.h"
 #include "Internationalization/FastDecimalFormat.h"
+#include "OnlineSubsystem.h"
 #import "OnlineStoreKitHelper.h"
 
 FOnlineStoreIOS::FOnlineStoreIOS(FOnlineSubsystemIOS* InSubsystem)
diff --git a/Engine/Plugins/Online/IOS/OnlineSubsystemIOS/Source/Private/OnlineStoreInterfaceIOS.cpp b/Engine/Plugins/Online/IOS/OnlineSubsystemIOS/Source/Private/OnlineStoreInterfaceIOS.cpp
--- a/Engine/Plugins/Online/IOS/OnlineSubsystemIOS/Source/Private/OnlineStoreInterfaceIOS.cpp
+++ b/Engine/Plugins/Online/IOS/OnlineSubsystemIOS/Source/Private/OnlineStoreInterfaceIOS.cpp
@@ -1,8 +1,9 @@
 // Copyright 1998-2019 Epic Games, Inc. All Rights Reserved.
 
-#include "OnlineSubsystemIOSPrivatePCH.h"
 #include "OnlineStoreInterfaceIOS.h"
 #include "OnlineSubsystemNames.h"
+#include "OnlineSubsystem.h"
+#include "Interfaces/OnlineStoreInterface.h"
 #import "OnlineStoreKitHelper.h"
 
 
diff --git a/Engine/Plugins/Online/IOS/OnlineSubsystemIOS/Source/Private/OnlineStoreKitHelper.cpp b/Engine/Plugins/Online/IOS/OnlineSubsystemIOS/Source/Private/OnlineStoreKitHelper.cpp
--- a/Engine/Plugins/Online/IOS/OnlineSubsystemIOS/Source/Private/OnlineStoreKitHelper.cpp
+++ b/Engine/Plugins/Online/IOS/OnlineSubsystemIOS/Source/Private/OnlineStoreKitHelper.cpp
@@ -1,7 +1,9 @@
 // Copyright 1998-2019 Epic Games, Inc. All Rights Reserved.
 
-#include "OnlineSubsystemIOSPrivatePCH.h"
-#import "OnlineStoreKitHelper.h"
+#include "OnlineStoreKitHelper.h"
+#include "Interfaces/OnlineStoreInterface.h"
+#include "OnlineSubsystemIOS.h"
+#include "IOS/IOSAppDelegate.h"
 
 /**
  * Convert an Apple SKPaymentTransaction receipt into a string
diff --git a/Engine/Plugins/Online/IOS/OnlineSubsystemIOS/Source/Private/OnlineStoreKitHelper.h b/Engine/Plugins/Online/IOS/OnlineSubsystemIOS/Source/Private/OnlineStoreKitHelper.h
--- a/Engine/Plugins/Online/IOS/OnlineSubsystemIOS/Source/Private/OnlineStoreKitHelper.h
+++ b/Engine/Plugins/Online/IOS/OnlineSubsystemIOS/Source/Private/OnlineStoreKitHelper.h
@@ -1,5 +1,7 @@
 // Copyright 1998-2019 Epic Games, Inc. All Rights Reserved.
 
+#pragma once
+
 // SK includes
 #import <StoreKit/SKRequest.h>
 #import <StoreKit/SKError.h>
@@ -11,6 +13,7 @@
 #import <StoreKit/SKPaymentQueue.h>
 
 #include "OnlineSubsystemTypes.h"
+#include "OnlineStoreIOS.h"
 
 /**
  * Holds in a common format the data that comes out of an SKPaymentTransaction
diff --git a/Engine/Plugins/Online/IOS/OnlineSubsystemIOS/Source/Private/OnlineSubsystemIOS.cpp b/Engine/Plugins/Online/IOS/OnlineSubsystemIOS/Source/Private/OnlineSubsystemIOS.cpp
--- a/Engine/Plugins/Online/IOS/OnlineSubsystemIOS/Source/Private/OnlineSubsystemIOS.cpp
+++ b/Engine/Plugins/Online/IOS/OnlineSubsystemIOS/Source/Private/OnlineSubsystemIOS.cpp
@@ -1,6 +1,7 @@
 // Copyright 1998-2019 Epic Games, Inc. All Rights Reserved.
 
-#include "OnlineSubsystemIOSPrivatePCH.h"
+#include "OnlineSubsystemIOS.h"
+#include "IOS/IOSAppDelegate.h"
 #import "OnlineStoreKitHelper.h"
 #import "OnlineAppStoreUtils.h"
 
diff --git a/Engine/Plugins/Online/IOS/OnlineSubsystemIOS/Source/Private/OnlineSubsystemIOSModule.cpp b/Engine/Plugins/Online/IOS/OnlineSubsystemIOS/Source/Private/OnlineSubsystemIOSModule.cpp
--- a/Engine/Plugins/Online/IOS/OnlineSubsystemIOS/Source/Private/OnlineSubsystemIOSModule.cpp
+++ b/Engine/Plugins/Online/IOS/OnlineSubsystemIOS/Source/Private/OnlineSubsystemIOSModule.cpp
@@ -1,7 +1,9 @@
 // Copyright 1998-2019 Epic Games, Inc. All Rights Reserved.
 
-#include "OnlineSubsystemIOSPrivatePCH.h"
+#include "OnlineSubsystemIOSModule.h"
 #include "HttpModule.h"
+#include "OnlineSubsystem.h"
+#include "OnlineSubsystemIOS.h"
 
 IMPLEMENT_MODULE( FOnlineSubsystemIOSModule, OnlineSubsystemIOS );
 
diff --git a/Engine/Plugins/Online/IOS/OnlineSubsystemIOS/Source/Private/OnlineSubsystemIOSPrivatePCH.h b/Engine/Plugins/Online/IOS/OnlineSubsystemIOS/Source/Private/OnlineSubsystemIOSPrivatePCH.h
deleted file mode 100644
diff --git a/Engine/Plugins/Online/IOS/OnlineSubsystemIOS/Source/Private/OnlineSubsystemIOSTypes.h b/Engine/Plugins/Online/IOS/OnlineSubsystemIOS/Source/Private/OnlineSubsystemIOSTypes.h
--- a/Engine/Plugins/Online/IOS/OnlineSubsystemIOS/Source/Private/OnlineSubsystemIOSTypes.h
+++ b/Engine/Plugins/Online/IOS/OnlineSubsystemIOS/Source/Private/OnlineSubsystemIOSTypes.h
@@ -5,6 +5,7 @@
 #include "OnlineSubsystemTypes.h"
 #include "OnlineSubsystemIOSPackage.h"
 
+class FOnlineSubsystemIOS;
 
 // from OnlineSubsystemTypes.h
 TEMP_UNIQUENETIDSTRING_SUBCLASS(FUniqueNetIdIOS, IOS_SUBSYSTEM);
diff --git a/Engine/Plugins/Online/IOS/OnlineSubsystemIOS/Source/Private/OnlineTurnBasedInterfaceIOS.cpp b/Engine/Plugins/Online/IOS/OnlineSubsystemIOS/Source/Private/OnlineTurnBasedInterfaceIOS.cpp
--- a/Engine/Plugins/Online/IOS/OnlineSubsystemIOS/Source/Private/OnlineTurnBasedInterfaceIOS.cpp
+++ b/Engine/Plugins/Online/IOS/OnlineSubsystemIOS/Source/Private/OnlineTurnBasedInterfaceIOS.cpp
@@ -1,6 +1,5 @@
 // Copyright 1998-2019 Epic Games, Inc. All Rights Reserved.
 
-#include "OnlineSubsystemIOSPrivatePCH.h"
 #include "OnlineTurnBasedInterfaceIOS.h"
 #include "TurnBasedEventListener.h"
 #include "Interfaces/OnlineIdentityInterface.h"
diff --git a/Engine/Plugins/Online/IOS/OnlineSubsystemIOS/Source/Private/OnlineUserCloudInterfaceIOS.cpp b/Engine/Plugins/Online/IOS/OnlineSubsystemIOS/Source/Private/OnlineUserCloudInterfaceIOS.cpp
--- a/Engine/Plugins/Online/IOS/OnlineSubsystemIOS/Source/Private/OnlineUserCloudInterfaceIOS.cpp
+++ b/Engine/Plugins/Online/IOS/OnlineSubsystemIOS/Source/Private/OnlineUserCloudInterfaceIOS.cpp
@@ -1,6 +1,5 @@
 // Copyright 1998-2019 Epic Games, Inc. All Rights Reserved.
 
-#include "OnlineSubsystemIOSPrivatePCH.h"
 #include "OnlineUserCloudInterfaceIOS.h"
 #include "HAL/PlatformProcess.h"
 
diff --git a/Engine/Plugins/Online/IOS/OnlineSubsystemIOS/Source/Private/TurnBasedEventListener.cpp b/Engine/Plugins/Online/IOS/OnlineSubsystemIOS/Source/Private/TurnBasedEventListener.cpp
--- a/Engine/Plugins/Online/IOS/OnlineSubsystemIOS/Source/Private/TurnBasedEventListener.cpp
+++ b/Engine/Plugins/Online/IOS/OnlineSubsystemIOS/Source/Private/TurnBasedEventListener.cpp
@@ -1,6 +1,5 @@
 // Copyright 1998-2019 Epic Games, Inc. All Rights Reserved.
 
-#include "OnlineSubsystemIOSPrivatePCH.h"
 #include "TurnBasedEventListener.h"
 
 @interface FTurnBasedEventListenerIOS()
diff --git a/Engine/Plugins/Online/IOS/OnlineSubsystemIOS/Source/Private/TurnBasedMatchmakerDelegateIOS.cpp b/Engine/Plugins/Online/IOS/OnlineSubsystemIOS/Source/Private/TurnBasedMatchmakerDelegateIOS.cpp
--- a/Engine/Plugins/Online/IOS/OnlineSubsystemIOS/Source/Private/TurnBasedMatchmakerDelegateIOS.cpp
+++ b/Engine/Plugins/Online/IOS/OnlineSubsystemIOS/Source/Private/TurnBasedMatchmakerDelegateIOS.cpp
@@ -1,9 +1,8 @@
 // Copyright 1998-2019 Epic Games, Inc. All Rights Reserved.
 
-#include "OnlineSubsystemIOSPrivatePCH.h"
-
-#include "Interfaces/OnlineTurnBasedInterface.h"
 #include "TurnBasedMatchmakerDelegateIOS.h"
+#include "OnlineTurnBasedInterfaceIOS.h"
+#include "Interfaces/OnlineTurnBasedInterface.h"
 
 @interface FTurnBasedMatchmakerDelegateIOS()
 {
diff --git a/Engine/Plugins/Online/IOS/OnlineSubsystemIOS/Source/Private/TurnBasedMatchmakerDelegateIOS.h b/Engine/Plugins/Online/IOS/OnlineSubsystemIOS/Source/Private/TurnBasedMatchmakerDelegateIOS.h
--- a/Engine/Plugins/Online/IOS/OnlineSubsystemIOS/Source/Private/TurnBasedMatchmakerDelegateIOS.h
+++ b/Engine/Plugins/Online/IOS/OnlineSubsystemIOS/Source/Private/TurnBasedMatchmakerDelegateIOS.h
@@ -3,6 +3,7 @@
 #pragma once
 
 #include <GameKit/GameKit.h>
+#include "Interfaces/OnlineTurnBasedInterface.h"
 
 @interface FTurnBasedMatchmakerDelegateIOS : NSObject < GKTurnBasedMatchmakerViewControllerDelegate >
 
diff --git a/Engine/Plugins/Online/IOS/OnlineSubsystemIOS/Source/Private/TurnBasedMatchmakerIOS.cpp b/Engine/Plugins/Online/IOS/OnlineSubsystemIOS/Source/Private/TurnBasedMatchmakerIOS.cpp
--- a/Engine/Plugins/Online/IOS/OnlineSubsystemIOS/Source/Private/TurnBasedMatchmakerIOS.cpp
+++ b/Engine/Plugins/Online/IOS/OnlineSubsystemIOS/Source/Private/TurnBasedMatchmakerIOS.cpp
@@ -1,6 +1,5 @@
 // Copyright 1998-2019 Epic Games, Inc. All Rights Reserved.
 
-#include "OnlineSubsystemIOSPrivatePCH.h"
 #include "TurnBasedMatchmakerIOS.h"
 
 #include <GameKit/GameKit.h>
diff --git a/Engine/Plugins/Online/IOS/OnlineSubsystemIOS/Source/Public/OnlinePurchaseIOS.h b/Engine/Plugins/Online/IOS/OnlineSubsystemIOS/Source/Public/OnlinePurchaseIOS.h
--- a/Engine/Plugins/Online/IOS/OnlineSubsystemIOS/Source/Public/OnlinePurchaseIOS.h
+++ b/Engine/Plugins/Online/IOS/OnlineSubsystemIOS/Source/Public/OnlinePurchaseIOS.h
@@ -3,8 +3,10 @@
 #pragma once
 
 #include "Interfaces/OnlinePurchaseInterface.h"
+#include "OnlineSubsystemIOSTypes.h"
 
 struct FStoreKitTransactionData;
+@class SKProductsResponse;
 @class FStoreKitHelperV2;
 
 /**
diff --git a/Engine/Plugins/Online/IOS/OnlineSubsystemIOS/Source/Public/OnlineSubsystemIOS.h b/Engine/Plugins/Online/IOS/OnlineSubsystemIOS/Source/Public/OnlineSubsystemIOS.h
--- a/Engine/Plugins/Online/IOS/OnlineSubsystemIOS/Source/Public/OnlineSubsystemIOS.h
+++ b/Engine/Plugins/Online/IOS/OnlineSubsystemIOS/Source/Public/OnlineSubsystemIOS.h
@@ -5,6 +5,18 @@
 #include "OnlineSubsystem.h"
 #include "OnlineSubsystemImpl.h"
 #include "OnlineSubsystemIOSPackage.h"
+#include "OnlineSessionInterfaceIOS.h"
+#include "OnlineFriendsInterfaceIOS.h"
+#include "OnlineIdentityInterfaceIOS.h"
+#include "OnlineLeaderboardsInterfaceIOS.h"
+#include "OnlineStoreInterfaceIOS.h"
+#include "OnlineStoreIOS.h"
+#include "OnlinePurchaseIOS.h"
+#include "OnlineAchievementsInterfaceIOS.h"
+#include "OnlineExternalUIInterfaceIOS.h"
+#include "OnlineTurnBasedInterfaceIOS.h"
+#include "OnlineUserCloudInterfaceIOS.h"
+#include "OnlineSharedCloudInterfaceIOS.h"
 
 @class FStoreKitHelperV2;
 @class FAppStoreUtils;
diff --git a/Engine/Source/Programs/UnrealBuildTool/Configuration/UEBuildModuleCPP.cs b/Engine/Source/Programs/UnrealBuildTool/Configuration/UEBuildModuleCPP.cs
--- a/Engine/Source/Programs/UnrealBuildTool/Configuration/UEBuildModuleCPP.cs
+++ b/Engine/Source/Programs/UnrealBuildTool/Configuration/UEBuildModuleCPP.cs
@@ -920,7 +920,6 @@ namespace UnrealBuildTool
 				case "MetalRHI":
 				case "PS4RHI":
                 case "Gnmx":
-				case "OnlineSubsystemIOS":
 					return true;
 			}
 			return false;
