Fix compile errors in CIS.

#rnx

diff --git a/Engine/Plugins/Online/IOS/OnlineSubsystemIOS/Source/Private/OnlineUserCloudInterfaceIOS.cpp b/Engine/Plugins/Online/IOS/OnlineSubsystemIOS/Source/Private/OnlineUserCloudInterfaceIOS.cpp
--- a/Engine/Plugins/Online/IOS/OnlineSubsystemIOS/Source/Private/OnlineUserCloudInterfaceIOS.cpp
+++ b/Engine/Plugins/Online/IOS/OnlineSubsystemIOS/Source/Private/OnlineUserCloudInterfaceIOS.cpp
@@ -1,5 +1,6 @@
 // Copyright 1998-2019 Epic Games, Inc. All Rights Reserved.
 
+#include "OnlineUserCloudInterfaceIOS.h"
 #include "HAL/PlatformProcess.h"
 #include "PlatformFeatures.h"
 
