fix for friends interface deprecation warning

diff --git a/Engine/Plugins/Online/IOS/OnlineSubsystemIOS/Source/Private/OnlineFriendsInterfaceIOS.cpp b/Engine/Plugins/Online/IOS/OnlineSubsystemIOS/Source/Private/OnlineFriendsInterfaceIOS.cpp
--- a/Engine/Plugins/Online/IOS/OnlineSubsystemIOS/Source/Private/OnlineFriendsInterfaceIOS.cpp
+++ b/Engine/Plugins/Online/IOS/OnlineSubsystemIOS/Source/Private/OnlineFriendsInterfaceIOS.cpp
@@ -51,6 +51,7 @@ FOnlineFriendsIOS::FOnlineFriendsIOS(FOnlineSubsystemIOS* InSubsystem)
 bool FOnlineFriendsIOS::ReadFriendsList(int32 LocalUserNum, const FString& ListName, const FOnReadFriendsListComplete& Delegate /*= FOnReadFriendsListComplete()*/)
 {
 	UE_LOG_ONLINE_FRIEND(Verbose, TEXT("FOnlineFriendsIOS::ReadFriendsList()"));
+#if __IPHONE_OS_VERSION_MIN_REQUIRED < __IPHONE_10_0
 	bool bSuccessfullyBeganReadFriends = false;
 
 	if( IdentityInterface->GetLocalGameCenterUser() != NULL && IdentityInterface->GetLocalGameCenterUser().isAuthenticated )
@@ -198,6 +199,10 @@ bool FOnlineFriendsIOS::ReadFriendsList(int32 LocalUserNum, const FString& ListN
 	}
 
 	return bSuccessfullyBeganReadFriends;
+#else
+	Delegate.ExecuteIfBound(LocalUserNum, false, ListName, FString(TEXT("ReadFriendsList() is not supported as of UE 4.21")));
+	return false;
+#endif
 }
 
 bool FOnlineFriendsIOS::DeleteFriendsList(int32 LocalUserNum, const FString& ListName, const FOnDeleteFriendsListComplete& Delegate /*= FOnDeleteFriendsListComplete()*/)
