SaveGame using iCloud for iOS and tvOS

UE-52763 tvOS/iOS: Allow save games to save to iCloud
UE-52346 TM-Gameplay: Load Level Volume Fails to Read TestSlot.sav file
UE-64066 TM-Gameplay Load Game does not work on TVOS

#iOS
#4.22

Adapted the Online iCloud plugin to work without enabling the Game Center plugin (using empty UniqueNetId as a parameter for FOnlineUserCloudInterfaceIOS methods and callbacks)

At plugin init, the data from the iCloud will overwrite the local storage.

Replaced saveRecord with CKModifyRecordsOperation to avoid errors on existing iCloud records - using savePolicy = CKRecordSaveAllKeys to overwrite the previously uploaded save file

Entitlements:
Fixed the problem with writing the iCloud lines in UE4Game.entitlements : The flag from DefaultEngine.ini was overwritten by IOSEngine.ini (same section & key). Same issue with the Game Center flag

Signing:
It looks like there is a problem (since Xcode 6/iOS8, when Apple has changed the format for the iCloud entitlements, in order to support CloudKit) with the Entitlements section from the provisioning file (downloaded from Apple Developper) <key>com.apple.developer.icloud-services</key> contains "*" instead of the  CloudKit and CloudDocuments entries. Fixed it, otherwise the app won't appear in the iPhone's Settings->iCloud apps menu

diff --git a/Engine/Config/BaseEngine.ini b/Engine/Config/BaseEngine.ini
--- a/Engine/Config/BaseEngine.ini
+++ b/Engine/Config/BaseEngine.ini
@@ -1843,6 +1843,7 @@ RSyncUsername=
 SSHPrivateKeyOverridePath=
 bEnableRemoteNotificationsSupport=False
 bEnableCloudKitSupport=False
+IOSCloudKitSyncStrategy=None
 bGenerateCrashReportSymbols=false
 bAutomaticSigning=false
 UseFastIntrinsics=False
diff --git a/Engine/Config/IOS/IOSEngine.ini b/Engine/Config/IOS/IOSEngine.ini
--- a/Engine/Config/IOS/IOSEngine.ini
+++ b/Engine/Config/IOS/IOSEngine.ini
@@ -33,10 +33,6 @@ DefaultProviderName=IOSAdvertising
 [DeviceProfileManager]
 DeviceProfileSelectionModule="IOSDeviceProfileSelector"
 
-[/Script/IOSRuntimeSettings.IOSRuntimeSettings]
-bEnableGameCenterSupport=false
-bEnableCloudKitSupport=false
-
 [OnlineSubsystemIOS.Store]
 bSupportsInAppPurchasing=true
 bUseStoreV2=true
diff --git a/Engine/Plugins/Online/IOS/OnlineSubsystemIOS/Source/OnlineSubsystemIOS.Build.cs b/Engine/Plugins/Online/IOS/OnlineSubsystemIOS/Source/OnlineSubsystemIOS.Build.cs
--- a/Engine/Plugins/Online/IOS/OnlineSubsystemIOS/Source/OnlineSubsystemIOS.Build.cs
+++ b/Engine/Plugins/Online/IOS/OnlineSubsystemIOS/Source/OnlineSubsystemIOS.Build.cs
@@ -10,8 +10,12 @@ public class OnlineSubsystemIOS : ModuleRules
 		PrivateIncludePaths.AddRange( 
             new string[] {
                 "Private",             
-                }
-                );
+                });
+
+        PublicIncludePaths.AddRange(
+            new string[] {
+                "Runtime/IOS/IOSPlatformFeatures/Public"
+                });
 
         PublicDefinitions.Add("ONLINESUBSYSTEMIOS_PACKAGE=1");
 
@@ -23,6 +27,7 @@ public class OnlineSubsystemIOS : ModuleRules
                 "Sockets",
 				"OnlineSubsystem", 
                 "Http",
+                "IOSPlatformFeatures",
             }
             );
 
diff --git a/Engine/Plugins/Online/IOS/OnlineSubsystemIOS/Source/Private/OnlineUserCloudInterfaceIOS.cpp b/Engine/Plugins/Online/IOS/OnlineSubsystemIOS/Source/Private/OnlineUserCloudInterfaceIOS.cpp
--- a/Engine/Plugins/Online/IOS/OnlineSubsystemIOS/Source/Private/OnlineUserCloudInterfaceIOS.cpp
+++ b/Engine/Plugins/Online/IOS/OnlineSubsystemIOS/Source/Private/OnlineUserCloudInterfaceIOS.cpp
@@ -1,8 +1,8 @@
 // Copyright 1998-2019 Epic Games, Inc. All Rights Reserved.
 
 #include "OnlineSubsystemIOSPrivatePCH.h"
-#include "OnlineUserCloudInterfaceIOS.h"
 #include "HAL/PlatformProcess.h"
+#include "PlatformFeatures.h"
 
 @implementation IOSCloudStorage
 
@@ -70,7 +70,14 @@
 			CKRecordID* recordId = [[[CKRecordID alloc] initWithRecordName:fileName] autorelease];
 			CKRecord* record = [[CKRecord alloc] initWithRecordType:@"file" recordID: recordId];
 			record[@"contents"] = fileContents;
-			[DB saveRecord : record completionHandler : handler];
+
+			// use CKModifyRecordsOperation to allow updating existing records 
+			CKModifyRecordsOperation *modifyRecords = [[CKModifyRecordsOperation alloc] initWithRecordsToSave:@[record] recordIDsToDelete:nil];
+			modifyRecords.savePolicy = CKRecordSaveAllKeys;
+			modifyRecords.qualityOfService = NSQualityOfServiceUserInitiated;
+			modifyRecords.perRecordCompletionBlock = handler;
+			[DB addOperation : modifyRecords];
+
 			return true;
 		}
 	}
@@ -312,7 +319,7 @@ void FOnlineUserCloudInterfaceIOS::GetUserFileList(const FUniqueNetId& UserId, T
 bool FOnlineUserCloudInterfaceIOS::ReadUserFile(const FUniqueNetId& UserId, const FString& FileName)
 {
 #ifdef __IPHONE_8_0
-	FCloudFile* CloudFile = GetCloudFile(FileName);
+	FCloudFile* CloudFile = GetCloudFile(FileName, true);
 	if (CloudFile)
 	{
         __block FString NewFile = FileName;
@@ -395,7 +402,7 @@ void FOnlineUserCloudInterfaceIOS::CancelWriteUserFile(const FUniqueNetId& UserI
 bool FOnlineUserCloudInterfaceIOS::DeleteUserFile(const FUniqueNetId& UserId, const FString& FileName, bool bShouldCloudDelete, bool bShouldLocallyDelete)
 {
 #ifdef __IPHONE_8_0
-	FCloudFile* CloudFile = GetCloudFile(FileName);
+	FCloudFile* CloudFile = GetCloudFile(FileName, true);
 	if (CloudFile)
 	{
         if (bShouldCloudDelete)
@@ -452,3 +459,152 @@ void FOnlineUserCloudInterfaceIOS::DumpCloudFileState(const FUniqueNetId& UserId
 {
 	// Not implemented
 }
+
+// cloud save implementation
+void FOnlineUserCloudInterfaceIOS::InitCloudSave(bool InIOSAlwaysSyncCloudFiles)
+{
+	SaveSystem = (FIOSSaveGameSystem*)IPlatformFeaturesModule::Get().GetSaveGameSystem();
+
+	check(SaveSystem);
+
+	bIOSAlwaysSyncCloudFiles = InIOSAlwaysSyncCloudFiles;
+
+	OnEnumerateUserCloudFilesCompleteDelegate = FOnEnumerateUserFilesCompleteDelegate::CreateRaw(this, &FOnlineUserCloudInterfaceIOS::OnEnumerateUserFilesComplete);
+	OnInitialFetchUserCloudFileCompleteDelegate = FOnReadUserFileCompleteDelegate::CreateRaw(this, &FOnlineUserCloudInterfaceIOS::OnInitialFetchUserCloudFileComplete);
+	OnWriteUserCloudFileCompleteDelegate = FOnWriteUserFileCompleteDelegate::CreateRaw(this, &FOnlineUserCloudInterfaceIOS::OnWriteUserCloudFileComplete);
+	OnReadUserCloudFileCompleteDelegate = FOnReadUserFileCompleteDelegate::CreateRaw(this, &FOnlineUserCloudInterfaceIOS::OnReadUserCloudFileComplete);
+	OnDeleteUserCloudFileCompleteDelegate = FOnDeleteUserFileCompleteDelegate::CreateRaw(this, &FOnlineUserCloudInterfaceIOS::OnDeleteUserCloudFileComplete);
+
+	// link delegates to the iOS save system
+	SaveSystem->OnWriteUserCloudFileBeginDelegate = FIOSSaveGameSystem::FOnWriteUserCloudFileBegin::CreateRaw(this, &FOnlineUserCloudInterfaceIOS::OnWriteUserCloudFileBegin);
+	SaveSystem->OnReadUserCloudFileBeginDelegate = FIOSSaveGameSystem::FOnReadUserCloudFileBegin::CreateRaw(this, &FOnlineUserCloudInterfaceIOS::OnReadUserCloudFileBegin);
+	SaveSystem->OnDeleteUserCloudFileBeginDelegate = FIOSSaveGameSystem::FOnDeleteUserCloudFileBegin::CreateRaw(this, &FOnlineUserCloudInterfaceIOS::OnDeleteUserCloudFileBegin);
+
+	UniqueNetId = MakeShareable(new FUniqueNetIdIOS(""));
+
+	UpdateDictionary = [NSMutableDictionary new];
+	
+	OnEnumerateUserCloudFilesCompleteDelegateHandle = AddOnEnumerateUserFilesCompleteDelegate_Handle(OnEnumerateUserCloudFilesCompleteDelegate);
+	EnumerateUserFiles(*UniqueNetId);
+}
+
+void FOnlineUserCloudInterfaceIOS::OnEnumerateUserFilesComplete(bool bWasSuccessful, const FUniqueNetId & UserId)
+{
+	ClearOnEnumerateUserFilesCompleteDelegate_Handle(OnEnumerateUserCloudFilesCompleteDelegateHandle);
+
+	TArray<FCloudFileHeader> UserFiles;
+	if (bWasSuccessful)
+	{
+		GetUserFileList(UserId, UserFiles);
+
+		for (int32 Idx = 0; Idx < UserFiles.Num(); Idx++)
+		{
+			OnReadUserCloudFileCompleteDelegateHandle = AddOnReadUserFileCompleteDelegate_Handle(OnInitialFetchUserCloudFileCompleteDelegate);
+			ReadUserFile(UserId, UserFiles[Idx].FileName);
+		}
+	}
+}
+
+void FOnlineUserCloudInterfaceIOS::OnWriteUserCloudFileComplete(bool bWasSuccessful, const FUniqueNetId & UserId, const FString & FileName)
+{
+	ClearOnWriteUserFileCompleteDelegate_Handle(OnWriteUserCloudFileCompleteDelegateHandle);
+
+	if (bWasSuccessful)
+	{
+		//flag that we have the latest record
+		[UpdateDictionary setObject : [NSNumber numberWithBool : false]  forKey : FileName.GetNSString()];
+	}
+	
+	// clean up temporary data
+	ClearFile(UserId, FileName);
+}
+
+void FOnlineUserCloudInterfaceIOS::OnInitialFetchUserCloudFileComplete(bool bWasSuccessful, const FUniqueNetId & UserId, const FString & FileName)
+{
+	OnReadUserCloudFileComplete(bWasSuccessful, UserId, FileName);
+
+	if (bWasSuccessful)
+	{
+		//flag that we have the latest record
+		[UpdateDictionary setObject : [NSNumber numberWithBool : false]  forKey : FileName.GetNSString()];
+	}
+		
+	// clean up temporary data
+	ClearFile(UserId, FileName);
+}
+
+void FOnlineUserCloudInterfaceIOS::OnReadUserCloudFileComplete(bool bWasSuccessful, const FUniqueNetId & UserId, const FString & FileName)
+{
+	ClearOnReadUserFileCompleteDelegate_Handle(OnReadUserCloudFileCompleteDelegateHandle);
+
+	if (bWasSuccessful)
+	{
+		// locally cache the server data
+		TArray<uint8> FileContents;
+		GetFileContents(UserId, FileName, FileContents);
+
+		check(SaveSystem);
+
+		SaveSystem->SaveGameNoCloud(*FileName, FileContents);
+		
+		//flag that we have the latest record
+		[UpdateDictionary setObject : [NSNumber numberWithBool : false]  forKey : FileName.GetNSString()];
+	}
+}
+
+void FOnlineUserCloudInterfaceIOS::OnDeleteUserCloudFileComplete(bool bWasSuccessful, const FUniqueNetId & UserId, const FString & FileName)
+{
+	ClearOnDeleteUserFileCompleteDelegate_Handle(OnDeleteUserCloudFileCompleteDelegateHandle);
+}
+
+void FOnlineUserCloudInterfaceIOS::OnWriteUserCloudFileBegin(const FString &  FileName, const TArray<uint8>& FileContents)
+{
+	OnWriteUserCloudFileCompleteDelegateHandle = AddOnWriteUserFileCompleteDelegate_Handle(OnWriteUserCloudFileCompleteDelegate);
+	WriteUserFile(*UniqueNetId, FileName, (TArray<uint8>&)FileContents);
+}
+
+bool FOnlineUserCloudInterfaceIOS::ShouldFetchRecordFromCloud(const FString &  FileName)
+{
+#if !PLATFORM_TVOS
+	if (!bIOSAlwaysSyncCloudFiles)
+	{
+		NSNumber *value = [UpdateDictionary objectForKey : FileName.GetNSString()];
+		if (value)
+		{
+			return[value boolValue];
+		}
+	}
+#endif
+	return true;
+}
+
+void FOnlineUserCloudInterfaceIOS::OnReadUserCloudFileBegin(const FString &  FileName, TArray<uint8>& FileContents)
+{
+#if !PLATFORM_TVOS
+	if (ShouldFetchRecordFromCloud(FileName))
+#endif
+	{
+		OnReadUserCloudFileCompleteDelegateHandle = AddOnReadUserFileCompleteDelegate_Handle(OnReadUserCloudFileCompleteDelegate);
+
+		ReadUserFile(*UniqueNetId, FileName);
+
+		FCloudFile* CloudFile = GetCloudFile(FileName, false);
+
+		while (CloudFile && CloudFile->AsyncState == EOnlineAsyncTaskState::InProgress)
+		{
+			FPlatformProcess::Sleep(0.01f);
+		}
+
+		GetFileContents(*UniqueNetId, FileName, FileContents);
+
+		// clean up temporary data
+		ClearFile(*UniqueNetId, FileName);
+	}
+}
+
+void FOnlineUserCloudInterfaceIOS::OnDeleteUserCloudFileBegin(const FString &  FileName)
+{
+	OnDeleteUserCloudFileCompleteDelegateHandle = AddOnDeleteUserFileCompleteDelegate_Handle(OnWriteUserCloudFileCompleteDelegate);
+
+	DeleteUserFile(*UniqueNetId, FileName, true, true); 
+}
\ No newline at end of file
diff --git a/Engine/Plugins/Online/IOS/OnlineSubsystemIOS/Source/Private/OnlineUserCloudInterfaceIOS.h b/Engine/Plugins/Online/IOS/OnlineSubsystemIOS/Source/Private/OnlineUserCloudInterfaceIOS.h
--- a/Engine/Plugins/Online/IOS/OnlineSubsystemIOS/Source/Private/OnlineUserCloudInterfaceIOS.h
+++ b/Engine/Plugins/Online/IOS/OnlineSubsystemIOS/Source/Private/OnlineUserCloudInterfaceIOS.h
@@ -3,6 +3,8 @@
 #pragma once
 
 #include "Interfaces/OnlineUserCloudInterface.h"
+#include "OnlineSubsystemIOSTypes.h"
+#include "IOSSaveGameSystem.h"
 
 #ifdef __IPHONE_8_0
 #import <CloudKit/CloudKit.h>
@@ -42,7 +44,7 @@ protected:
 	bool ClearCloudFile(const FString& FileName);
 
 public:
-    FOnlineUserCloudInterfaceIOS(){ MetaDataState = EOnlineAsyncTaskState::Done; }
+	FOnlineUserCloudInterfaceIOS();
 	virtual ~FOnlineUserCloudInterfaceIOS();
 
 	// IOnlineUserCloud
@@ -60,6 +62,8 @@ public:
 	virtual void DumpCloudFileState(const FUniqueNetId& UserId, const FString& FileName) override;
 
 private:
+	/** Enable Cloud support - from Settings->iOS->Online */
+	bool bEnableCloudKitSupport;
 	/** File metadata */
 	TArray<struct FCloudFileHeader> CloudMetaData;
     /** File metadata query state */
@@ -69,6 +73,76 @@ private:
 
     /** Critical section for thread safe operation on cloud files */
     FCriticalSection CloudDataLock;
+
+	/** Reference to the iOS file save system */
+	FIOSSaveGameSystem* SaveSystem;
+
+	/** Flag from Settings->iOS: Always read from the iCloud on LoadGame */
+	bool bIOSAlwaysSyncCloudFiles;
+
+	/** Delegates to various cloud functionality triggered */
+	FOnEnumerateUserFilesCompleteDelegate OnEnumerateUserCloudFilesCompleteDelegate;
+	FOnWriteUserFileCompleteDelegate OnWriteUserCloudFileCompleteDelegate;
+	FOnReadUserFileCompleteDelegate OnReadUserCloudFileCompleteDelegate;
+	FOnDeleteUserFileCompleteDelegate OnDeleteUserCloudFileCompleteDelegate;
+
+	/** Handles to those delegates */
+	FDelegateHandle OnEnumerateUserCloudFilesCompleteDelegateHandle;
+	FDelegateHandle OnWriteUserCloudFileCompleteDelegateHandle;
+	FDelegateHandle OnReadUserCloudFileCompleteDelegateHandle;
+	FDelegateHandle OnDeleteUserCloudFileCompleteDelegateHandle;
+	
+	/** Cache the UserId */
+	TSharedPtr<FUniqueNetIdIOS> UniqueNetId;
+
+	/**
+	 *	Delegate triggered when all user files have been enumerated
+	 * @param bWasSuccessful did the operation complete successfully
+	 * @param UserId user that triggered the operation
+	 */
+	void OnEnumerateUserFilesComplete(bool bWasSuccessful, const FUniqueNetId& UserId);
+
+	/**
+	 *	Delegate triggered for each user cloud file written
+	 * @param bWasSuccessful did the operation complete successfully
+	 * @param UserId user that triggered the operation
+	 * @param FileName filename written to cloud
+	 */
+	void OnWriteUserCloudFileComplete(bool bWasSuccessful, const FUniqueNetId& UserId, const FString& FileName);
+
+	/**
+	 *	Delegate triggered for each user cloud file read - will overwrite the local files
+	 * @param bWasSuccessful did the operation complete successfully
+	 * @param UserId user that triggered the operation
+	 * @param FileName filename read from cloud
+	 */
+	void OnReadUserCloudFileComplete(bool bWasSuccessful, const FUniqueNetId& UserId, const FString& FileName);
+
+	/**
+	 *	Delegate triggered for each user cloud file deleted
+	 * @param bWasSuccessful did the operation complete successfully
+	 * @param UserId user that triggered the operation
+	 * @param FileName filename deleted from cloud
+	 */
+	void OnDeleteUserCloudFileComplete(bool bWasSuccessful, const FUniqueNetId& UserId, const FString& FileName);
+
+	/**
+	 *	Delegate in the iOS file save system
+	 * Called in SaveGame
+	 */
+	void OnWriteUserCloudFileBegin(const FString&  FileName, const TArray<uint8>& FileContents);
+
+	/**
+	 *	Delegate in the iOS file save system
+	 * Called in ReadGame
+	 */
+	void OnReadUserCloudFileBegin(const FString&  FileName, const TArray<uint8>& FileContents);
+
+	/**
+	 *	Delegate in the iOS file save system
+	 * Called in DeleteGame
+	 */
+	void OnDeleteUserCloudFileBegin(const FString&  FileName);
 };
 
 typedef TSharedPtr<FOnlineUserCloudInterfaceIOS, ESPMode::ThreadSafe> FOnlineUserCloudIOSPtr;
diff --git a/Engine/Source/Programs/IOS/iPhonePackager/MobileProvisionUtilities.cs b/Engine/Source/Programs/IOS/iPhonePackager/MobileProvisionUtilities.cs
--- a/Engine/Source/Programs/IOS/iPhonePackager/MobileProvisionUtilities.cs
+++ b/Engine/Source/Programs/IOS/iPhonePackager/MobileProvisionUtilities.cs
@@ -1,4 +1,4 @@
-﻿/**
+/**
  * Copyright 1998-2019 Epic Games, Inc. All Rights Reserved.
  */
 
@@ -315,6 +315,27 @@ namespace iPhonePackager
 				XCentPList.SetValueForKey("keychain-access-groups", KeyGroups);
 			}
 
+			// must have CloudKit and CloudDocuments for com.apple.developer.icloud-services
+			// otherwise the game will not be listed in the Settings->iCloud apps menu on the device
+			{
+				List<string> ServicesGroups = XCentPList.GetArray("com.apple.developer.icloud-services", "string");
+
+				if (ServicesGroups.Count == 0 || !ServicesGroups[0].StartsWith("Cloud"))
+				{
+					ServicesGroups.Clear();
+
+					string ServicesString;
+					XCentPList.GetString("com.apple.developer.icloud-services", out ServicesString);
+
+					if (ServicesString.Equals("*"))
+					{
+						ServicesGroups.Add("CloudKit");
+						ServicesGroups.Add("CloudDocuments");
+					}
+				}
+				XCentPList.SetValueForKey("com.apple.developer.icloud-services", ServicesGroups);
+			}
+
 			return XCentPList.SaveToString();
 		}
 
diff --git a/Engine/Source/Programs/UnrealBuildTool/Platform/IOS/UEBuildIOS.cs b/Engine/Source/Programs/UnrealBuildTool/Platform/IOS/UEBuildIOS.cs
--- a/Engine/Source/Programs/UnrealBuildTool/Platform/IOS/UEBuildIOS.cs
+++ b/Engine/Source/Programs/UnrealBuildTool/Platform/IOS/UEBuildIOS.cs
@@ -980,6 +980,16 @@ namespace UnrealBuildTool
 		public override void ModifyModuleRulesForActivePlatform(string ModuleName, ModuleRules Rules, ReadOnlyTargetRules Target)
 		{
 		}
+		
+		/// <summary>
+		/// Setup the binaries for this specific platform.
+		/// </summary>
+		/// <param name="Target">The target being built</param>
+		/// <param name="ExtraModuleNames"></param>
+		public override void AddExtraModules(ReadOnlyTargetRules Target, List<string> ExtraModuleNames)
+		{
+			ExtraModuleNames.Add("IOSPlatformFeatures");
+		}
 
 		/// <summary>
 		/// Creates a toolchain instance for the given platform.
diff --git a/Engine/Source/Programs/UnrealBuildTool/Platform/IOS/UEDeployIOS.cs b/Engine/Source/Programs/UnrealBuildTool/Platform/IOS/UEDeployIOS.cs
--- a/Engine/Source/Programs/UnrealBuildTool/Platform/IOS/UEDeployIOS.cs
+++ b/Engine/Source/Programs/UnrealBuildTool/Platform/IOS/UEDeployIOS.cs
@@ -1219,7 +1219,14 @@ namespace UnrealBuildTool
 				Text.AppendLine("\t<key>com.apple.developer.icloud-services</key>");
 				Text.AppendLine("\t<array>");
 				Text.AppendLine("\t\t<string>CloudKit</string>");
+				Text.AppendLine("\t\t<string>CloudDocuments</string>");
 				Text.AppendLine("\t</array>");
+				Text.AppendLine("\t<key>com.apple.developer.ubiquity-container-identifiers</key>");
+				Text.AppendLine("\t<array>");
+				Text.AppendLine("\t\t<string>iCloud.$(CFBundleIdentifier)</string>");
+				Text.AppendLine("\t</array>");
+				Text.AppendLine("\t<key>com.apple.developer.ubiquity-kvstore-identifier</key>");
+				Text.AppendLine("\t<string>$(TeamIdentifierPrefix)$(CFBundleIdentifier)</string>");
 			}
 
 			bool bRemoteNotificationsSupported = false;
diff --git a/Engine/Source/Runtime/Core/Private/IOS/IOSPlatformMisc.cpp b/Engine/Source/Runtime/Core/Private/IOS/IOSPlatformMisc.cpp
--- a/Engine/Source/Runtime/Core/Private/IOS/IOSPlatformMisc.cpp
+++ b/Engine/Source/Runtime/Core/Private/IOS/IOSPlatformMisc.cpp
@@ -151,6 +151,12 @@ void FIOSPlatformMisc::PlatformInit()
 	}
 }
 
+// Defines the PlatformFeatures module name for iOS, used by PlatformFeatures.h.
+const TCHAR* FIOSPlatformMisc::GetPlatformFeaturesModuleName()
+{
+	return TEXT("IOSPlatformFeatures");
+}
+
 void FIOSPlatformMisc::PlatformHandleSplashScreen(bool ShowSplashScreen)
 {
     if (GShowSplashScreen != ShowSplashScreen)
diff --git a/Engine/Source/Runtime/Core/Public/IOS/IOSPlatformMisc.h b/Engine/Source/Runtime/Core/Public/IOS/IOSPlatformMisc.h
--- a/Engine/Source/Runtime/Core/Public/IOS/IOSPlatformMisc.h
+++ b/Engine/Source/Runtime/Core/Public/IOS/IOSPlatformMisc.h
@@ -20,6 +20,7 @@ struct CORE_API FIOSPlatformMisc : public FApplePlatformMisc
     static void PlatformPreInit();
 	static void PlatformInit();
     static void PlatformHandleSplashScreen(bool ShowSplashScreen = false);
+	static const TCHAR* GetPlatformFeaturesModuleName();
 
 	FORCEINLINE static int32 GetMaxPathLength()
 	{
diff --git a/Engine/Source/Runtime/IOS/IOSPlatformFeatures/IOSPlatformFeatures.Build.cs b/Engine/Source/Runtime/IOS/IOSPlatformFeatures/IOSPlatformFeatures.Build.cs
new file mode 100644
--- /dev/null
+++ b/Engine/Source/Runtime/IOS/IOSPlatformFeatures/IOSPlatformFeatures.Build.cs
@@ -0,0 +1,15 @@
+// Copyright 1998-2019 Epic Games, Inc. All Rights Reserved.
+
+using UnrealBuildTool;
+
+public class IOSPlatformFeatures : ModuleRules
+{
+	public IOSPlatformFeatures(ReadOnlyTargetRules Target) : base(Target)
+	{
+		PrivateDependencyModuleNames.AddRange(new string[] 
+			{ 
+				"Core", 
+				"Engine" 
+			});
+	}
+}
diff --git a/Engine/Source/Runtime/IOS/IOSPlatformFeatures/Private/IOSPlatformFeatures.cpp b/Engine/Source/Runtime/IOS/IOSPlatformFeatures/Private/IOSPlatformFeatures.cpp
new file mode 100644
--- /dev/null
+++ b/Engine/Source/Runtime/IOS/IOSPlatformFeatures/Private/IOSPlatformFeatures.cpp
@@ -0,0 +1,17 @@
+// Copyright 1998-2019 Epic Games, Inc. All Rights Reserved.
+
+#include "CoreMinimal.h"
+#include "PlatformFeatures.h"
+#include "IOSSaveGameSystem.h"
+
+class FIOSPlatformFeatures : public IPlatformFeaturesModule
+{
+public:
+	virtual class ISaveGameSystem* GetSaveGameSystem() override
+	{
+		static FIOSSaveGameSystem IOSSaveGameSystem;
+		return &IOSSaveGameSystem;
+	}
+};
+
+IMPLEMENT_MODULE(FIOSPlatformFeatures, IOSPlatformFeatures);
diff --git a/Engine/Source/Runtime/IOS/IOSPlatformFeatures/Private/IOSSaveGameSystem.cpp b/Engine/Source/Runtime/IOS/IOSPlatformFeatures/Private/IOSSaveGameSystem.cpp
new file mode 100644
--- /dev/null
+++ b/Engine/Source/Runtime/IOS/IOSPlatformFeatures/Private/IOSSaveGameSystem.cpp
@@ -0,0 +1,79 @@
+// Copyright 1998-2019 Epic Games, Inc. All Rights Reserved.
+
+#include "IOSSaveGameSystem.h"
+#include "GameDelegates.h"
+
+DEFINE_LOG_CATEGORY_STATIC(LogIOSSaveGame, Log, All);
+
+//
+// Implementation members
+//
+
+FIOSSaveGameSystem::FIOSSaveGameSystem()
+{
+	Initialize();
+}
+
+FIOSSaveGameSystem::~FIOSSaveGameSystem()
+{
+	Shutdown();
+}
+
+void FIOSSaveGameSystem::Initialize()
+{
+}
+
+void FIOSSaveGameSystem::Shutdown()
+{
+}
+
+ISaveGameSystem::ESaveExistsResult FIOSSaveGameSystem::DoesSaveGameExistWithResult(const TCHAR* Name, const int32 UserIndex)
+{
+	if (IFileManager::Get().FileSize(*GetSaveGamePath(Name)) >= 0)
+	{
+		return ESaveExistsResult::OK;
+	}
+	return ESaveExistsResult::DoesNotExist;
+}
+
+bool FIOSSaveGameSystem::SaveGameNoCloud(const TCHAR* Name, const TArray<uint8>& Data)
+{
+	return FFileHelper::SaveArrayToFile(Data, *GetSaveGamePath(Name));
+}
+
+bool FIOSSaveGameSystem::SaveGame(bool bAttemptToUseUI, const TCHAR* Name, const int32 UserIndex, const TArray<uint8>& Data)
+{
+	// send to the iCloud, if enabled
+	OnWriteUserCloudFileBeginDelegate.ExecuteIfBound(FString(Name), Data);
+
+	return FFileHelper::SaveArrayToFile(Data, *GetSaveGamePath(Name));
+}
+
+bool FIOSSaveGameSystem::LoadGame(bool bAttemptToUseUI, const TCHAR* Name, const int32 UserIndex, TArray<uint8>& Data)
+{
+	// try to read it from the iCloud, if enabled
+	OnReadUserCloudFileBeginDelegate.ExecuteIfBound(FString(Name), Data);
+
+	if (Data.Num() > 0)
+	{
+		// we've received data from the iCloud, the save file was overwritten
+		return true;
+	}
+
+	// no iCloud data, read from the local storage
+	return FFileHelper::LoadFileToArray(Data, *GetSaveGamePath(Name));
+}
+
+bool FIOSSaveGameSystem::DeleteGame(bool bAttemptToUseUI, const TCHAR* Name, const int32 UserIndex)
+{
+	// delete the file from the iCloud
+	OnDeleteUserCloudFileBeginDelegate.ExecuteIfBound(FString(Name));
+
+	// delete the file from the local storage
+	return IFileManager::Get().Delete(*GetSaveGamePath(Name), true, false, !bAttemptToUseUI);
+}
+
+FString FIOSSaveGameSystem::GetSaveGamePath(const TCHAR* Name)
+{
+	return FString::Printf(TEXT("%s""SaveGames/%s.sav"), *FPaths::ProjectSavedDir(), Name);
+}
diff --git a/Engine/Source/Runtime/IOS/IOSPlatformFeatures/Public/IOSSaveGameSystem.h b/Engine/Source/Runtime/IOS/IOSPlatformFeatures/Public/IOSSaveGameSystem.h
new file mode 100644
--- /dev/null
+++ b/Engine/Source/Runtime/IOS/IOSPlatformFeatures/Public/IOSSaveGameSystem.h
@@ -0,0 +1,70 @@
+// Copyright 1998-2019 Epic Games, Inc. All Rights Reserved.
+
+#pragma once
+
+#include "SaveGameSystem.h"
+
+class FIOSSaveGameSystem : public ISaveGameSystem
+{
+public:
+	FIOSSaveGameSystem();
+	virtual ~FIOSSaveGameSystem();
+
+	// ISaveGameSystem interface
+	virtual bool PlatformHasNativeUI() override
+	{
+		return false;
+	}
+
+	virtual bool DoesSaveGameExist(const TCHAR* Name, const int32 UserIndex) override
+	{
+		return ESaveExistsResult::OK == DoesSaveGameExistWithResult(Name, UserIndex);
+	}
+
+	virtual ESaveExistsResult DoesSaveGameExistWithResult(const TCHAR* Name, const int32 UserIndex) override;
+
+	/**
+	 * Called on the initial iCloud sync
+	 */
+	virtual bool SaveGameNoCloud(const TCHAR* Name, const TArray<uint8>& Data);
+	
+	virtual bool SaveGame(bool bAttemptToUseUI, const TCHAR* Name, const int32 UserIndex, const TArray<uint8>& Data) override;
+	/**
+	 * Called when writing the savegame file
+	 * send the file to iCloud, if enabled
+	 */
+	DECLARE_DELEGATE_TwoParams(FOnWriteUserCloudFileBegin, const FString&, const TArray<uint8>&);
+	FOnWriteUserCloudFileBegin OnWriteUserCloudFileBeginDelegate;
+
+	virtual bool LoadGame(bool bAttemptToUseUI, const TCHAR* Name, const int32 UserIndex, TArray<uint8>& Data) override;
+	/**
+	 * Called when reading the savegame file
+	 * read the file from iCloud, if enabled
+	 */
+	DECLARE_DELEGATE_TwoParams(FOnReadUserCloudFileBegin, const FString&, TArray<uint8>&);
+	FOnReadUserCloudFileBegin OnReadUserCloudFileBeginDelegate;
+
+	virtual bool DeleteGame(bool bAttemptToUseUI, const TCHAR* Name, const int32 UserIndex) override;
+	/**
+	 * Called when deleting the savegame file
+	 * delete the file from iCloud, if enabled
+	 */
+	DECLARE_DELEGATE_OneParam(FOnDeleteUserCloudFileBegin, const FString&);
+	FOnDeleteUserCloudFileBegin OnDeleteUserCloudFileBeginDelegate;
+
+private:
+	/**
+	 * Initializes the SaveData library then loads and initializes the SaveDialog library
+	 */
+	void Initialize();
+
+	/**
+	 * Terminates and unloads the SaveDialog library then terminates the SaveData library
+	 */
+	void Shutdown();
+
+	/**
+	 * Get the path to save game file for the given name
+	 */
+	virtual FString GetSaveGamePath(const TCHAR* Name);
+};
diff --git a/Engine/Source/Runtime/IOS/IOSRuntimeSettings/Classes/IOSRuntimeSettings.h b/Engine/Source/Runtime/IOS/IOSRuntimeSettings/Classes/IOSRuntimeSettings.h
--- a/Engine/Source/Runtime/IOS/IOSRuntimeSettings/Classes/IOSRuntimeSettings.h
+++ b/Engine/Source/Runtime/IOS/IOSRuntimeSettings/Classes/IOSRuntimeSettings.h
@@ -82,6 +82,19 @@ enum class EIOSLandscapeOrientation : uint8
 	LandscapeRight = 1 UMETA(DisplayName = "Landscape (right home button)"),
 };
 
+UENUM()
+enum class EIOSCloudKitSyncStrategy : uint8
+{
+	/** Only at game start */
+	None  = 0 UMETA(DisplayName = "Never (do not use iCloud for Load/Save Game)"),
+
+	/** Only at game start */
+	OnlyAtGameStart = 1 UMETA(DisplayName = "At game start only (iOS)"),
+
+	/** Always */
+	Always = 2 UMETA(DisplayName = "Always (whenever LoadGame is called)"),
+};
+
 /**
  *	IOS Build resource file struct, used to serialize filepaths to the configs for use in the build system,
  */
@@ -205,6 +218,10 @@ public:
 	UPROPERTY(GlobalConfig, EditAnywhere, Category = Online)
 	uint32 bEnableCloudKitSupport : 1;
 
+	// iCloud Read stategy
+	UPROPERTY(GlobalConfig, EditAnywhere, Category = Online, meta = (DisplayName = "iCloud save files sync strategy"), meta = (EditCondition = "bEnableCloudKitSupport"))
+	EIOSCloudKitSyncStrategy IOSCloudKitSyncStrategy;
+
     // Should push/remote notifications support (iOS Online Subsystem) be enabled?
     UPROPERTY(GlobalConfig, EditAnywhere, Category = Online)
     uint32 bEnableRemoteNotificationsSupport : 1;
