Missing files from CL4760235

SaveGame using iCloud for iOS and tvOS

#ios
#4.22

diff --git a/Engine/Plugins/Online/IOS/OnlineSubsystemIOS/Source/Private/OnlineSubsystemIOS.cpp b/Engine/Plugins/Online/IOS/OnlineSubsystemIOS/Source/Private/OnlineSubsystemIOS.cpp
--- a/Engine/Plugins/Online/IOS/OnlineSubsystemIOS/Source/Private/OnlineSubsystemIOS.cpp
+++ b/Engine/Plugins/Online/IOS/OnlineSubsystemIOS/Source/Private/OnlineSubsystemIOS.cpp
@@ -179,6 +179,17 @@ bool FOnlineSubsystemIOS::Init()
 			StoreInterface = MakeShareable(new FOnlineStoreInterfaceIOS());
 		}
 	}
+
+	if (UserCloudInterface && IsCloudKitEnabled())
+	{
+		FString IOSCloudKitSyncStrategy = "";
+		GConfig->GetString(TEXT("/Script/IOSRuntimeSettings.IOSRuntimeSettings"), TEXT("IOSCloudKitSyncStrategy"), IOSCloudKitSyncStrategy, GEngineIni);
+		
+		if (!IOSCloudKitSyncStrategy.Equals("None"))
+		{
+			UserCloudInterface->InitCloudSave(IOSCloudKitSyncStrategy.Equals("Always"));
+		}
+	}
 	
 	InitAppStoreHelper();
 
@@ -317,15 +328,23 @@ bool FOnlineSubsystemIOS::IsEnabled() const
 	bool bEnableGameCenter = false;
 	GConfig->GetBool(TEXT("/Script/IOSRuntimeSettings.IOSRuntimeSettings"), TEXT("bEnableGameCenterSupport"), bEnableGameCenter, GEngineIni);
 
-	bool bEnableCloudKit = false;
-    GConfig->GetBool(TEXT("/Script/IOSRuntimeSettings.IOSRuntimeSettings"), TEXT("bEnableCloudKitSupport"), bEnableCloudKit, GEngineIni);
-	
+	const bool bEnableCloudKit = IsCloudKitEnabled();
+
 	const bool bIsInAppPurchasingEnabled = IsInAppPurchasingEnabled();
 	const bool bIsEnabledByConfig = FOnlineSubsystemImpl::IsEnabled(); // TODO: Do we want to enable this by this config?
 	
 	return bEnableGameCenter || bEnableCloudKit || bIsInAppPurchasingEnabled || bIsEnabledByConfig;
 }
 
+bool FOnlineSubsystemIOS::IsCloudKitEnabled()
+{
+	bool bEnableCloudKit;
+	GConfig->GetBool(TEXT("/Script/IOSRuntimeSettings.IOSRuntimeSettings"), TEXT("bEnableCloudKitSupport"), bEnableCloudKit, GEngineIni);
+
+	return bEnableCloudKit;
+}
+
+
 bool FOnlineSubsystemIOS::IsV2StoreEnabled()
 {
 	bool bUseStoreV2 = false;
diff --git a/Engine/Plugins/Online/IOS/OnlineSubsystemIOS/Source/Private/OnlineUserCloudInterfaceIOS.h b/Engine/Plugins/Online/IOS/OnlineSubsystemIOS/Source/Private/OnlineUserCloudInterfaceIOS.h
--- a/Engine/Plugins/Online/IOS/OnlineSubsystemIOS/Source/Private/OnlineUserCloudInterfaceIOS.h
+++ b/Engine/Plugins/Online/IOS/OnlineSubsystemIOS/Source/Private/OnlineUserCloudInterfaceIOS.h
@@ -61,9 +61,9 @@ public:
 	virtual void DumpCloudState(const FUniqueNetId& UserId) override;
 	virtual void DumpCloudFileState(const FUniqueNetId& UserId, const FString& FileName) override;
 
+	// Initialize Cloud saving
+	void InitCloudSave(bool InIOSAlwaysSyncCloudFiles);
 private:
-	/** Enable Cloud support - from Settings->iOS->Online */
-	bool bEnableCloudKitSupport;
 	/** File metadata */
 	TArray<struct FCloudFileHeader> CloudMetaData;
     /** File metadata query state */
@@ -80,8 +80,15 @@ private:
 	/** Flag from Settings->iOS: Always read from the iCloud on LoadGame */
 	bool bIOSAlwaysSyncCloudFiles;
 
+	/** Store the iCloud sync status for each save file
+	* entris of type (string) filename: (bool) synced with iCloud
+	* updated by the silent notifications, if enabled
+	*/
+	NSMutableDictionary* UpdateDictionary;
+	
 	/** Delegates to various cloud functionality triggered */
 	FOnEnumerateUserFilesCompleteDelegate OnEnumerateUserCloudFilesCompleteDelegate;
+	FOnReadUserFileCompleteDelegate OnInitialFetchUserCloudFileCompleteDelegate;
 	FOnWriteUserFileCompleteDelegate OnWriteUserCloudFileCompleteDelegate;
 	FOnReadUserFileCompleteDelegate OnReadUserCloudFileCompleteDelegate;
 	FOnDeleteUserFileCompleteDelegate OnDeleteUserCloudFileCompleteDelegate;
@@ -102,6 +109,14 @@ private:
 	 */
 	void OnEnumerateUserFilesComplete(bool bWasSuccessful, const FUniqueNetId& UserId);
 
+	/**
+	 *	Delegate triggered on the init for each user cloud file read - will overwrite the local files
+	 * @param bWasSuccessful did the operation complete successfully
+	 * @param UserId user that triggered the operation
+	 * @param FileName filename read from cloud
+	 */
+	void OnInitialFetchUserCloudFileComplete(bool bWasSuccessful, const FUniqueNetId& UserId, const FString& FileName);
+
 	/**
 	 *	Delegate triggered for each user cloud file written
 	 * @param bWasSuccessful did the operation complete successfully
@@ -136,13 +151,19 @@ private:
 	 *	Delegate in the iOS file save system
 	 * Called in ReadGame
 	 */
-	void OnReadUserCloudFileBegin(const FString&  FileName, const TArray<uint8>& FileContents);
+	void OnReadUserCloudFileBegin(const FString&  FileName, TArray<uint8>& FileContents);
 
 	/**
 	 *	Delegate in the iOS file save system
 	 * Called in DeleteGame
 	 */
 	void OnDeleteUserCloudFileBegin(const FString&  FileName);
+
+	/**
+	 *	Returns true if the record must be fetched from the iCloud
+	 */
+	 bool ShouldFetchRecordFromCloud(const FString &  FileName);
+
 };
 
 typedef TSharedPtr<FOnlineUserCloudInterfaceIOS, ESPMode::ThreadSafe> FOnlineUserCloudIOSPtr;
diff --git a/Engine/Plugins/Online/IOS/OnlineSubsystemIOS/Source/Public/OnlineSubsystemIOS.h b/Engine/Plugins/Online/IOS/OnlineSubsystemIOS/Source/Public/OnlineSubsystemIOS.h
--- a/Engine/Plugins/Online/IOS/OnlineSubsystemIOS/Source/Public/OnlineSubsystemIOS.h
+++ b/Engine/Plugins/Online/IOS/OnlineSubsystemIOS/Source/Public/OnlineSubsystemIOS.h
@@ -68,6 +68,12 @@ PACKAGE_SCOPE:
 	 */
 	static bool IsInAppPurchasingEnabled();
 	
+	/**
+	 * Is CloudKit enabled
+	 * @return true if enabled, false otherwise
+	 */
+	static bool IsCloudKitEnabled();
+	
 	/**
 	 * Is Store v2 enabled (disabling legacy store interface)
 	 * @return true if enabled, false otherwise
