Missing files from CL4760235

SaveGame using iCloud for iOS and tvOS

#ios
#4.22

diff --git a/Engine/Plugins/Online/IOS/OnlineSubsystemIOS/Source/Private/OnlineUserCloudInterfaceIOS.h b/Engine/Plugins/Online/IOS/OnlineSubsystemIOS/Source/Private/OnlineUserCloudInterfaceIOS.h
--- a/Engine/Plugins/Online/IOS/OnlineSubsystemIOS/Source/Private/OnlineUserCloudInterfaceIOS.h
+++ b/Engine/Plugins/Online/IOS/OnlineSubsystemIOS/Source/Private/OnlineUserCloudInterfaceIOS.h
@@ -44,7 +44,7 @@ protected:
 	bool ClearCloudFile(const FString& FileName);
 
 public:
-	FOnlineUserCloudInterfaceIOS();
+    FOnlineUserCloudInterfaceIOS() : SaveSystem(NULL), UpdateDictionary(NULL) { MetaDataState = EOnlineAsyncTaskState::Done; }
 	virtual ~FOnlineUserCloudInterfaceIOS();
 
 	// IOnlineUserCloud
