SessionSettings update added to FOnlineSessionIOS

diff --git a/Engine/Plugins/Online/IOS/OnlineSubsystemIOS/Source/Private/OnlineSessionInterfaceIOS.cpp b/Engine/Plugins/Online/IOS/OnlineSubsystemIOS/Source/Private/OnlineSessionInterfaceIOS.cpp
--- a/Engine/Plugins/Online/IOS/OnlineSubsystemIOS/Source/Private/OnlineSessionInterfaceIOS.cpp
+++ b/Engine/Plugins/Online/IOS/OnlineSubsystemIOS/Source/Private/OnlineSessionInterfaceIOS.cpp
@@ -477,7 +477,15 @@ bool FOnlineSessionIOS::UpdateSession(FName SessionName, FOnlineSessionSettings&
 	bool bSuccessfullyUpdatedSession = false;
 	
 	UE_LOG_ONLINE_SESSION(Display, TEXT("FOnlineSessionIOS::UpdateSession - not implemented"));
-	
+
+	FNamedOnlineSession* Session = GetNamedSession(SessionName);
+	if (Session)
+	{
+		Session->SessionSettings = UpdatedSessionSettings;
+
+		bSuccessfullyUpdatedSession = true;
+	}
+
 	TriggerOnUpdateSessionCompleteDelegates(SessionName, bSuccessfullyUpdatedSession);
 	
 	return bSuccessfullyUpdatedSession;
