Fixed crash when loading iCloud/CloudKit SaveGame on iOS

diff --git a/Engine/Plugins/Online/IOS/OnlineSubsystemIOS/Source/Private/OnlineUserCloudInterfaceIOS.cpp b/Engine/Plugins/Online/IOS/OnlineSubsystemIOS/Source/Private/OnlineUserCloudInterfaceIOS.cpp
--- a/Engine/Plugins/Online/IOS/OnlineSubsystemIOS/Source/Private/OnlineUserCloudInterfaceIOS.cpp
+++ b/Engine/Plugins/Online/IOS/OnlineSubsystemIOS/Source/Private/OnlineUserCloudInterfaceIOS.cpp
@@ -340,7 +340,10 @@ bool FOnlineUserCloudInterfaceIOS::ReadUserFile(const FUniqueNetId& UserId, cons
 			if (error)
 			{
 				// TODO: record is potentially not found
-				File->AsyncState = EOnlineAsyncTaskState::Failed;
+				if (File)
+				{
+					File->AsyncState = EOnlineAsyncTaskState::Failed;
+				}
 				TriggerOnReadUserFileCompleteDelegates(false, UserId, NewFile);
 				NSLog(@"Error: %@", error);
 			}
@@ -348,9 +351,12 @@ bool FOnlineUserCloudInterfaceIOS::ReadUserFile(const FUniqueNetId& UserId, cons
 			{
  				// store the contents in the memory record database
 				NSData* data = (NSData*)record[@"contents"];
-				File->Data.Empty();
-				File->Data.Append((uint8*)data.bytes, data.length);
-				File->AsyncState = EOnlineAsyncTaskState::Done;
+				if (File && data)
+				{
+					File->Data.Empty();
+					File->Data.Append((uint8*)data.bytes, data.length);
+					File->AsyncState = EOnlineAsyncTaskState::Done;
+				}
 				TriggerOnReadUserFileCompleteDelegates(true, UserId, NewFile);
 				NSLog(@"Record Read!");
 			}
@@ -380,15 +386,21 @@ bool FOnlineUserCloudInterfaceIOS::WriteUserFile(const FUniqueNetId& UserId, con
 			if (error)
 			{
 				// TODO: record is potentially newer on the server
-				File->AsyncState = EOnlineAsyncTaskState::Failed;
+				if (File)
+				{
+					File->AsyncState = EOnlineAsyncTaskState::Failed;
+				}
 				TriggerOnWriteUserFileCompleteDelegates(false, UserId, NewFile);
 				NSLog(@"Error: %@", error);
 			}
 			else
 			{
                 FCloudFileHeader* Header = GetCloudFileHeader(NewFile, true);
-				File->Data = DataContents;
-				File->AsyncState = EOnlineAsyncTaskState::Done;
+				if (File)
+				{
+					File->Data = DataContents;
+					File->AsyncState = EOnlineAsyncTaskState::Done;
+				}
 				TriggerOnWriteUserFileCompleteDelegates(true, UserId, NewFile);
 				NSLog(@"Record Saved!");
 			}
@@ -613,7 +625,7 @@ void FOnlineUserCloudInterfaceIOS::OnReadUserCloudFileBegin(const FString &  Fil
 
 void FOnlineUserCloudInterfaceIOS::OnDeleteUserCloudFileBegin(const FString &  FileName)
 {
-	OnDeleteUserCloudFileCompleteDelegateHandle = AddOnDeleteUserFileCompleteDelegate_Handle(OnWriteUserCloudFileCompleteDelegate);
+	OnDeleteUserCloudFileCompleteDelegateHandle = AddOnDeleteUserFileCompleteDelegate_Handle(OnDeleteUserCloudFileCompleteDelegate);
 
 	DeleteUserFile(*UniqueNetId, FileName, true, true); 
 }
\ No newline at end of file
