Import the monolithic gamekit to avoid issues with the dependencies changing between SDK versions.

diff --git a/Engine/Plugins/Online/IOS/OnlineSubsystemIOS/Source/Private/OnlineIdentityInterfaceIOS.h b/Engine/Plugins/Online/IOS/OnlineSubsystemIOS/Source/Private/OnlineIdentityInterfaceIOS.h
--- a/Engine/Plugins/Online/IOS/OnlineSubsystemIOS/Source/Private/OnlineIdentityInterfaceIOS.h
+++ b/Engine/Plugins/Online/IOS/OnlineSubsystemIOS/Source/Private/OnlineIdentityInterfaceIOS.h
@@ -3,7 +3,7 @@
 
 #pragma once
 
-#import <GameKit/GKLocalPlayer.h>
+#import <GameKit/GameKit.h>
 
 #include "OnlineSubsystemIOSTypes.h"
 #include "Interfaces/OnlineIdentityInterface.h"
