Fix non-unity build

diff --git a/Engine/Plugins/Online/IOS/OnlineSubsystemIOS/Source/Private/OnlineIdentityInterfaceIOS.h b/Engine/Plugins/Online/IOS/OnlineSubsystemIOS/Source/Private/OnlineIdentityInterfaceIOS.h
--- a/Engine/Plugins/Online/IOS/OnlineSubsystemIOS/Source/Private/OnlineIdentityInterfaceIOS.h
+++ b/Engine/Plugins/Online/IOS/OnlineSubsystemIOS/Source/Private/OnlineIdentityInterfaceIOS.h
@@ -3,6 +3,7 @@
 
 #pragma once
 
+#import <UIKit/UIKit.h>
 #import <GameKit/GameKit.h>
 
 #include "OnlineSubsystemIOSTypes.h"
diff --git a/Engine/Plugins/Online/IOS/OnlineSubsystemIOS/Source/Private/OnlineSessionInterfaceIOS.h b/Engine/Plugins/Online/IOS/OnlineSubsystemIOS/Source/Private/OnlineSessionInterfaceIOS.h
--- a/Engine/Plugins/Online/IOS/OnlineSubsystemIOS/Source/Private/OnlineSessionInterfaceIOS.h
+++ b/Engine/Plugins/Online/IOS/OnlineSubsystemIOS/Source/Private/OnlineSessionInterfaceIOS.h
@@ -2,6 +2,7 @@
 
 #pragma once
 
+#import <UIKit/UIKit.h>
 #import <GameKit/GKLocalPlayer.h>
 
 #include "Interfaces/OnlineSessionInterface.h"
