Ask To Buy - Redemption fix for iOS

diff --git a/Engine/Plugins/Online/IOS/OnlineSubsystemIOS/Source/Private/OnlinePurchaseIOS.cpp b/Engine/Plugins/Online/IOS/OnlineSubsystemIOS/Source/Private/OnlinePurchaseIOS.cpp
--- a/Engine/Plugins/Online/IOS/OnlineSubsystemIOS/Source/Private/OnlinePurchaseIOS.cpp
+++ b/Engine/Plugins/Online/IOS/OnlineSubsystemIOS/Source/Private/OnlinePurchaseIOS.cpp
@@ -31,12 +31,12 @@ FOnlinePurchaseIOS::FOnlinePurchaseIOS(FOnlineSubsystemIOS* InSubsystem)
 	, bRestoringTransactions(false)
 	, Subsystem(InSubsystem)
 {
-	UE_LOG_ONLINE_PURCHASE(Verbose, TEXT( "FOnlinePurchaseIOS::FOnlinePurchaseIOS" ));
+	UE_LOG_ONLINE_PURCHASE(Log, TEXT( "FOnlinePurchaseIOS::FOnlinePurchaseIOS" ));
 }
 
 FOnlinePurchaseIOS::FOnlinePurchaseIOS()
 {
-	UE_LOG_ONLINE_PURCHASE(Verbose, TEXT( "FOnlinePurchaseIOS::FOnlinePurchaseIOS" ));
+	UE_LOG_ONLINE_PURCHASE(Log, TEXT( "FOnlinePurchaseIOS::FOnlinePurchaseIOS" ));
 }
 
 FOnlinePurchaseIOS::~FOnlinePurchaseIOS()
@@ -67,6 +67,12 @@ void FOnlinePurchaseIOS::InitStoreKit(FStoreKitHelperV2* InStoreKit)
 	[StoreHelper AddOnTransactionDeferred: OnTransactionDeferredDelegate];
 }
 
+void FOnlinePurchaseIOS::ManuallyIteratePaymentQueue()
+{
+	NSArray *Transactions = [[SKPaymentQueue defaultQueue] transactions];
+	[StoreHelper paymentQueue : [SKPaymentQueue defaultQueue] updatedTransactions : Transactions];
+}
+
 bool FOnlinePurchaseIOS::IsAllowedToPurchase(const FUniqueNetId& UserId)
 {
 	bool bCanMakePurchases = [SKPaymentQueue canMakePayments];
@@ -155,7 +161,7 @@ void FOnlinePurchaseIOS::Checkout(const FUniqueNetId& UserId, const FPurchaseChe
 
 void FOnlinePurchaseIOS::FinalizePurchase(const FUniqueNetId& UserId, const FString& ReceiptId)
 {
-	UE_LOG_ONLINE_PURCHASE(Verbose, TEXT("FOnlinePurchaseIOS::FinalizePurchase %s %s"), *UserId.ToString(), *ReceiptId);
+	UE_LOG_ONLINE_PURCHASE(Log, TEXT("FOnlinePurchaseIOS::FinalizePurchase %s %s"), *UserId.ToString(), *ReceiptId);
 
 	const FString ReceiptIdCopy(ReceiptId);
 	dispatch_async(dispatch_get_main_queue(), ^
@@ -197,17 +203,32 @@ void FOnlinePurchaseIOS::QueryReceipts(const FUniqueNetId& UserId, bool bRestore
 		}
 		else
 		{
-			UE_LOG_ONLINE_PURCHASE(Verbose, TEXT("FOnlinePurchaseIOS::QueryReceipts already restoring transactions"));
+			UE_LOG_ONLINE_PURCHASE(Log, TEXT("FOnlinePurchaseIOS::QueryReceipts already restoring transactions"));
 			bSuccess = false;
 		}
 	}
+	else
+	{
+		// We don't always seem to get events from our payment queue observer, 
+		// so manually iterate the transactions in the queue and handle them.
+		ManuallyIteratePaymentQueue();
+	}
 	
 	if (bTriggerDelegate)
 	{
 		// Query receipts comes dynamically from the StoreKit observer
-		Subsystem->ExecuteNextTick([Delegate, bSuccess]() {
-			FOnlineError Result(bSuccess);
-			Delegate.ExecuteIfBound(Result);
+		// Re-entrant WaitNextTick happening here because we want to wait 2 frames, to ensure
+		// Async delegates fired as a result of ManuallyIteratePaymentQueue have been handled.
+		TWeakPtr<FOnlinePurchaseIOS, ESPMode::ThreadSafe> WeakThis(AsShared());
+		Subsystem->ExecuteNextTick([WeakThis, Delegate, bSuccess]() {
+			FOnlinePurchaseIOSPtr StrongThis = WeakThis.Pin();
+			if (StrongThis.IsValid())
+			{
+				StrongThis->Subsystem->ExecuteNextTick([Delegate, bSuccess]() {
+					FOnlineError Result(bSuccess);
+					Delegate.ExecuteIfBound(Result);
+				});
+			}
 		});
 	}
 }
@@ -235,12 +256,12 @@ void FOnlinePurchaseIOS::GetReceipts(const FUniqueNetId& UserId, TArray<FPurchas
 
 void FOnlinePurchaseIOS::OnProductPurchaseRequestResponse(SKProductsResponse* Response, const FOnQueryOnlineStoreOffersComplete& CompletionDelegate)
 {
-	UE_LOG_ONLINE_PURCHASE(Verbose, TEXT("FOnlinePurchaseIOS::OnProductPurchaseRequestResponse"));
+	UE_LOG_ONLINE_PURCHASE(Log, TEXT("FOnlinePurchaseIOS::OnProductPurchaseRequestResponse"));
 }
 
 void FOnlinePurchaseIOS::OnTransactionCompleteResponse(EPurchaseTransactionState Result, const FStoreKitTransactionData& TransactionData)
 {
-	UE_LOG_ONLINE_PURCHASE(Verbose, TEXT("FOnlinePurchaseIOS::OnTransactionCompleteResponse %d %s"), (int32)Result, *TransactionData.ToDebugString());
+	UE_LOG_ONLINE_PURCHASE(Log, TEXT("FOnlinePurchaseIOS::OnTransactionCompleteResponse %d %s"), (int32)Result, *TransactionData.ToDebugString());
 	
 	FString UserIdStr = IOSUSER;
 	const TSharedRef<FOnlinePurchasePendingTransactionIOS>* UserPendingTransactionPtr = PendingTransactions.Find(UserIdStr);
diff --git a/Engine/Plugins/Online/IOS/OnlineSubsystemIOS/Source/Private/OnlineStoreKitHelper.cpp b/Engine/Plugins/Online/IOS/OnlineSubsystemIOS/Source/Private/OnlineStoreKitHelper.cpp
--- a/Engine/Plugins/Online/IOS/OnlineSubsystemIOS/Source/Private/OnlineStoreKitHelper.cpp
+++ b/Engine/Plugins/Online/IOS/OnlineSubsystemIOS/Source/Private/OnlineStoreKitHelper.cpp
@@ -4,6 +4,7 @@
 #include "Interfaces/OnlineStoreInterface.h"
 #include "OnlineSubsystemIOS.h"
 #include "IOS/IOSAppDelegate.h"
+#include "Misc/CommandLine.h"
 
 /**
  * Convert an Apple SKPaymentTransaction receipt into a string
@@ -131,27 +132,34 @@ FStoreKitTransactionData::FStoreKitTransactionData(const SKPaymentTransaction* T
 		switch ([transaction transactionState])
 		{
 			case SKPaymentTransactionStatePurchased:
-				UE_LOG_ONLINE_STORE(Verbose, TEXT("FStoreKitHelper::completeTransaction"));
-				[self completeTransaction : transaction];
+				if (FParse::Param(FCommandLine::Get(), TEXT("disableiosredeem")))
+				{
+					UE_LOG_ONLINE_STORE(Log, TEXT("FStoreKitHelper::completeTransaction (disabled)"));
+				}
+				else
+				{
+					UE_LOG_ONLINE_STORE(Log, TEXT("FStoreKitHelper::completeTransaction"));
+					[self completeTransaction : transaction];
+				}
 				break;
 			case SKPaymentTransactionStateFailed:
-				UE_LOG_ONLINE_STORE(Verbose, TEXT("FStoreKitHelper::failedTransaction"));
+				UE_LOG_ONLINE_STORE(Log, TEXT("FStoreKitHelper::failedTransaction"));
 				[self failedTransaction : transaction];
 				break;
 			case SKPaymentTransactionStateRestored:
-				UE_LOG_ONLINE_STORE(Verbose, TEXT("FStoreKitHelper::restoreTransaction"));
+				UE_LOG_ONLINE_STORE(Log, TEXT("FStoreKitHelper::restoreTransaction"));
 				[self restoreTransaction : transaction];
 				break;
 			case SKPaymentTransactionStatePurchasing:
-				UE_LOG_ONLINE_STORE(Verbose, TEXT("FStoreKitHelper::purchasingInProgress"));
+				UE_LOG_ONLINE_STORE(Log, TEXT("FStoreKitHelper::purchasingInProgress"));
 				[self purchaseInProgress : transaction];
 				continue;
 			case SKPaymentTransactionStateDeferred:
-				UE_LOG_ONLINE_STORE(Verbose, TEXT("FStoreKitHelper::purchaseDeferred"));
+				UE_LOG_ONLINE_STORE(Log, TEXT("FStoreKitHelper::purchaseDeferred"));
 				[self purchaseDeferred : transaction];
 				continue;
 			default:
-				UE_LOG_ONLINE_STORE(Verbose, TEXT("FStoreKitHelper::other: %d"), [transaction transactionState]);
+				UE_LOG_ONLINE_STORE(Log, TEXT("FStoreKitHelper::other: %d"), [transaction transactionState]);
 				break;
 		}
 	}
@@ -316,7 +324,7 @@ FStoreKitTransactionData::FStoreKitTransactionData(const SKPaymentTransaction* T
 
 -(void)requestProductData: (NSMutableSet*)productIDs
 {
-	UE_LOG_ONLINE_STORE(Verbose, TEXT("FStoreKitHelper::requestProductData"));
+	UE_LOG_ONLINE_STORE(Log, TEXT("FStoreKitHelper::requestProductData"));
 
 	Request = [[SKProductsRequest alloc] initWithProductIdentifiers:productIDs];
 	Request.delegate = self;
@@ -326,7 +334,7 @@ FStoreKitTransactionData::FStoreKitTransactionData(const SKPaymentTransaction* T
 
 -(void)makePurchase: (NSMutableSet*)productIDs
 {
-	UE_LOG_ONLINE_STORE(Verbose, TEXT("FStoreKitHelper::makePurchase"));
+	UE_LOG_ONLINE_STORE(Log, TEXT("FStoreKitHelper::makePurchase"));
 
 	Request = [[SKProductsRequest alloc] initWithProductIdentifiers:productIDs];
 	Request.delegate = self;
@@ -336,7 +344,7 @@ FStoreKitTransactionData::FStoreKitTransactionData(const SKPaymentTransaction* T
 
 -(void)productsRequest: (SKProductsRequest *)request didReceiveResponse : (SKProductsResponse *)response
 {
-	UE_LOG_ONLINE_STORE(Verbose, TEXT("FStoreKitHelper::didReceiveResponse"));
+	UE_LOG_ONLINE_STORE(Log, TEXT("FStoreKitHelper::didReceiveResponse"));
 	// Direct the response back to the store interface
 	[FIOSAsyncTask CreateTaskWithBlock : ^ bool(void)
 	{
@@ -367,6 +375,7 @@ FStoreKitTransactionData::FStoreKitTransactionData(const SKPaymentTransaction* T
 	{
 		[self paymentQueue : [SKPaymentQueue defaultQueue]  restoreCompletedTransactionsFailedWithError : error];
 		[Request release];
+		Request = nullptr;
 	}
 #endif
 }
@@ -437,7 +446,7 @@ FStoreKitTransactionData::FStoreKitTransactionData(const SKPaymentTransaction* T
 
 -(void)makePurchase:(NSArray*)products WithUserId: (const FString&) userId SimulateAskToBuy: (bool) bAskToBuy;
 {
-	UE_LOG_ONLINE_STOREV2(Verbose, TEXT("FStoreKitHelperV2::makePurchase by SKProduct with UserId"));
+	UE_LOG_ONLINE_STOREV2(Log, TEXT("FStoreKitHelperV2::makePurchase by SKProduct with UserId"));
 	
 	for (SKProduct* Product in products)
 	{
@@ -456,14 +465,14 @@ FStoreKitTransactionData::FStoreKitTransactionData(const SKPaymentTransaction* T
 
 -(void)makePurchase:(NSArray*)products SimulateAskToBuy: (bool) bAskToBuy
 {
-	UE_LOG_ONLINE_STOREV2(Verbose, TEXT("FStoreKitHelperV2::makePurchase by SKProduct"));
+	UE_LOG_ONLINE_STOREV2(Log, TEXT("FStoreKitHelperV2::makePurchase by SKProduct"));
 	FString EmptyString;
 	[self makePurchase: products WithUserId: EmptyString SimulateAskToBuy: bAskToBuy];
 }
 
 -(void)requestProductData: (NSMutableSet*)productIDs WithDelegate : (const FOnQueryOnlineStoreOffersComplete&)delegate
 {
-	UE_LOG_ONLINE_STOREV2(Verbose, TEXT("FStoreKitHelperV2::requestProductData"));
+	UE_LOG_ONLINE_STOREV2(Log, TEXT("FStoreKitHelperV2::requestProductData"));
 	
 	FSKProductsRequestHelper* TempRequest = [[FSKProductsRequestHelper alloc] initWithProductIdentifiers:productIDs];
 	TempRequest.OfferDelegate = delegate;
@@ -477,7 +486,7 @@ FStoreKitTransactionData::FStoreKitTransactionData(const SKPaymentTransaction* T
 {
 	if (request == self.Request)
 	{
-		UE_LOG_ONLINE_STOREV2(Verbose, TEXT("FStoreKitHelperV2::didReceiveResponse"));
+		UE_LOG_ONLINE_STOREV2(Log, TEXT("FStoreKitHelperV2::didReceiveResponse"));
 		
 		if ([request isKindOfClass : [FSKProductsRequestHelper class]])
 		{
@@ -600,7 +609,7 @@ FStoreKitTransactionData::FStoreKitTransactionData(const SKPaymentTransaction* T
 
 -(void)finalizeTransaction: (const FString&) receiptId
 {
-	UE_LOG_ONLINE_STOREV2(Verbose, TEXT("FStoreKitHelperV2::finalizeTransaction - %s"), *receiptId);
+	UE_LOG_ONLINE_STOREV2(Log, TEXT("FStoreKitHelperV2::finalizeTransaction - %s"), *receiptId);
 	for (SKPaymentTransaction* pendingTransaction in self.PendingTransactions)
 	{
 		if (pendingTransaction)
diff --git a/Engine/Plugins/Online/IOS/OnlineSubsystemIOS/Source/Public/OnlinePurchaseIOS.h b/Engine/Plugins/Online/IOS/OnlineSubsystemIOS/Source/Public/OnlinePurchaseIOS.h
--- a/Engine/Plugins/Online/IOS/OnlineSubsystemIOS/Source/Public/OnlinePurchaseIOS.h
+++ b/Engine/Plugins/Online/IOS/OnlineSubsystemIOS/Source/Public/OnlinePurchaseIOS.h
@@ -104,6 +104,8 @@ public:
 	/** Initialize the FStoreKitHelper for interaction with the app store */
 	void InitStoreKit(FStoreKitHelperV2* InStoreKit);
 	
+	/** Manually iterate the payment transaction queue and fire off appropriate callbacks in StoreKitHelper. */
+	void ManuallyIteratePaymentQueue();
 	
 private:
 	
