Fix for Game Center popup appeared even when deativated on behalf of fyi's.
Local testing in Match3 did not show Game Center popups following the ticket repro steps.

diff --git a/Engine/Plugins/Online/IOS/OnlineSubsystemIOS/Source/Private/OnlineSubsystemIOS.cpp b/Engine/Plugins/Online/IOS/OnlineSubsystemIOS/Source/Private/OnlineSubsystemIOS.cpp
--- a/Engine/Plugins/Online/IOS/OnlineSubsystemIOS/Source/Private/OnlineSubsystemIOS.cpp
+++ b/Engine/Plugins/Online/IOS/OnlineSubsystemIOS/Source/Private/OnlineSubsystemIOS.cpp
@@ -151,48 +151,53 @@ bool FOnlineSubsystemIOS::Init()
 	
 	if( !IsEnabled() )
 	{
-		UE_LOG_ONLINE(Warning, TEXT("GameCenter has been disabled in the system settings"));
+		UE_LOG_ONLINE(Warning, TEXT("All iOS online features have been disabled in the system settings"));
 		bSuccessfullyStartedUp = false;
 	}
 	else
 	{
 		SessionInterface = MakeShareable(new FOnlineSessionIOS(this));
 		IdentityInterface = MakeShareable(new FOnlineIdentityIOS(this));
-		FriendsInterface = MakeShareable(new FOnlineFriendsIOS(this));
-		LeaderboardsInterface = MakeShareable(new FOnlineLeaderboardsIOS(this));
-		AchievementsInterface = MakeShareable(new FOnlineAchievementsIOS(this));
-		ExternalUIInterface = MakeShareable(new FOnlineExternalUIIOS(this));
-        TurnBasedInterface = MakeShareable(new FOnlineTurnBasedIOS());
-        UserCloudInterface = MakeShareable(new FOnlineUserCloudInterfaceIOS());
-        SharedCloudInterface = MakeShareable(new FOnlineSharedCloudInterfaceIOS());
-	}
-	
-	if(IsInAppPurchasingEnabled())
-	{
-		if (IsV2StoreEnabled())
+
+		if (IsGameCenterEnabled())
 		{
-			StoreV2Interface = MakeShareable(new FOnlineStoreIOS(this));
-			PurchaseInterface = MakeShareable(new FOnlinePurchaseIOS(this));
-			InitStoreKitHelper();
+			FriendsInterface = MakeShareable(new FOnlineFriendsIOS(this));
+			LeaderboardsInterface = MakeShareable(new FOnlineLeaderboardsIOS(this));
+			AchievementsInterface = MakeShareable(new FOnlineAchievementsIOS(this));
+			ExternalUIInterface = MakeShareable(new FOnlineExternalUIIOS(this));
+			TurnBasedInterface = MakeShareable(new FOnlineTurnBasedIOS());
 		}
-		else
+
+		UserCloudInterface = MakeShareable(new FOnlineUserCloudInterfaceIOS());
+		SharedCloudInterface = MakeShareable(new FOnlineSharedCloudInterfaceIOS());
+
+		if (IsInAppPurchasingEnabled())
 		{
-			StoreInterface = MakeShareable(new FOnlineStoreInterfaceIOS());
+			if (IsV2StoreEnabled())
+			{
+				StoreV2Interface = MakeShareable(new FOnlineStoreIOS(this));
+				PurchaseInterface = MakeShareable(new FOnlinePurchaseIOS(this));
+				InitStoreKitHelper();
+			}
+			else
+			{
+				StoreInterface = MakeShareable(new FOnlineStoreInterfaceIOS());
+			}
 		}
-	}
 
-	if (UserCloudInterface && IsCloudKitEnabled())
-	{
-		FString IOSCloudKitSyncStrategy = "";
-		GConfig->GetString(TEXT("/Script/IOSRuntimeSettings.IOSRuntimeSettings"), TEXT("IOSCloudKitSyncStrategy"), IOSCloudKitSyncStrategy, GEngineIni);
-		
-		if (!IOSCloudKitSyncStrategy.Equals("None"))
+		if (UserCloudInterface && IsCloudKitEnabled())
 		{
-			UserCloudInterface->InitCloudSave(IOSCloudKitSyncStrategy.Equals("Always"));
+			FString IOSCloudKitSyncStrategy = "";
+			GConfig->GetString(TEXT("/Script/IOSRuntimeSettings.IOSRuntimeSettings"), TEXT("IOSCloudKitSyncStrategy"), IOSCloudKitSyncStrategy, GEngineIni);
+
+			if (!IOSCloudKitSyncStrategy.Equals("None"))
+			{
+				UserCloudInterface->InitCloudSave(IOSCloudKitSyncStrategy.Equals("Always"));
+			}
 		}
+
+		InitAppStoreHelper();
 	}
-	
-	InitAppStoreHelper();
 
 	return bSuccessfullyStartedUp;
 }
@@ -326,8 +331,7 @@ bool FOnlineSubsystemIOS::HandlePurchaseExecCommands(UWorld* InWorld, const TCHA
 
 bool FOnlineSubsystemIOS::IsEnabled() const
 {
-	bool bEnableGameCenter = false;
-	GConfig->GetBool(TEXT("/Script/IOSRuntimeSettings.IOSRuntimeSettings"), TEXT("bEnableGameCenterSupport"), bEnableGameCenter, GEngineIni);
+	bool bEnableGameCenter = IsGameCenterEnabled();
 
 	const bool bEnableCloudKit = IsCloudKitEnabled();
 
@@ -337,6 +341,14 @@ bool FOnlineSubsystemIOS::IsEnabled() const
 	return bEnableGameCenter || bEnableCloudKit || bIsInAppPurchasingEnabled || bIsEnabledByConfig;
 }
 
+bool FOnlineSubsystemIOS::IsGameCenterEnabled()
+{
+	bool bEnableGameCenter = false;
+	GConfig->GetBool(TEXT("/Script/IOSRuntimeSettings.IOSRuntimeSettings"), TEXT("bEnableGameCenterSupport"), bEnableGameCenter, GEngineIni);
+
+	return bEnableGameCenter;
+}
+
 bool FOnlineSubsystemIOS::IsCloudKitEnabled()
 {
 	bool bEnableCloudKit;
diff --git a/Engine/Plugins/Online/IOS/OnlineSubsystemIOS/Source/Public/OnlineSubsystemIOS.h b/Engine/Plugins/Online/IOS/OnlineSubsystemIOS/Source/Public/OnlineSubsystemIOS.h
--- a/Engine/Plugins/Online/IOS/OnlineSubsystemIOS/Source/Public/OnlineSubsystemIOS.h
+++ b/Engine/Plugins/Online/IOS/OnlineSubsystemIOS/Source/Public/OnlineSubsystemIOS.h
@@ -81,6 +81,12 @@ PACKAGE_SCOPE:
 	 */
 	static bool IsInAppPurchasingEnabled();
 	
+	/**
+	 * Is GameCenter enabled
+	 * @return true if enabled, false otherwise
+	 */
+	static bool IsGameCenterEnabled();
+	
 	/**
 	 * Is CloudKit enabled
 	 * @return true if enabled, false otherwise
