Fix Deprecation warning on tvOS 13 #rb jack.porter #jira UE-111501

diff --git a/Engine/Plugins/Online/IOS/OnlineSubsystemIOS/Source/Private/OnlineSubsystemIOS.cpp b/Engine/Plugins/Online/IOS/OnlineSubsystemIOS/Source/Private/OnlineSubsystemIOS.cpp
--- a/Engine/Plugins/Online/IOS/OnlineSubsystemIOS/Source/Private/OnlineSubsystemIOS.cpp
+++ b/Engine/Plugins/Online/IOS/OnlineSubsystemIOS/Source/Private/OnlineSubsystemIOS.cpp
@@ -7,12 +7,6 @@
 #include "Misc/ConfigCacheIni.h"
 #include "Stats/Stats.h"
 
-#if (defined(__IPHONE_OS_VERSION_MIN_REQUIRED) && __IPHONE_OS_VERSION_MIN_REQUIRED <= __IPHONE_12_0) || (defined(__TV_OS_VERSION_MIN_REQUIRED) && __TV_OS_VERSION_MIN_REQUIRED <= __TVOS_12_0)
-#define USE_DEPRECATED_PLAYERID 1
-#else
-#define USE_DEPRECATED_PLAYERID 0
-#endif
-
 FOnlineSubsystemIOS::FOnlineSubsystemIOS(FName InInstanceName)
 	: FOnlineSubsystemImpl(IOS_SUBSYSTEM, InInstanceName)
 	, StoreHelper(nil)
@@ -386,33 +380,11 @@ bool FOnlineSubsystemIOS::IsInAppPurchasingEnabled()
 
 NSString* FOnlineSubsystemIOS::GetPlayerId(GKPlayer* Player)
 {
-#if USE_DEPRECATED_PLAYERID
-	if ([GKPlayer respondsToSelector:@selector(gamePlayerID)] == YES)
-	{
-		return Player.gamePlayerID;
-	}
-	else
-	{
-		return Player.playerID;
-	}
-#else
 	return Player.gamePlayerID;
-#endif
 }
 
 NSString* FOnlineSubsystemIOS::GetPlayerId(GKLocalPlayer* Player)
 {
-#if USE_DEPRECATED_PLAYERID
-	if ([GKLocalPlayer respondsToSelector:@selector(gamePlayerID)] == YES)
-	{
-		return Player.gamePlayerID;
-	}
-	else
-	{
-		return Player.playerID;
-	}
-#else
 	return Player.gamePlayerID;
-#endif
 }
 
diff --git a/Engine/Source/Runtime/Core/Private/Apple/ApplePlatformDebugEvents.cpp b/Engine/Source/Runtime/Core/Private/Apple/ApplePlatformDebugEvents.cpp
--- a/Engine/Source/Runtime/Core/Private/Apple/ApplePlatformDebugEvents.cpp
+++ b/Engine/Source/Runtime/Core/Private/Apple/ApplePlatformDebugEvents.cpp
@@ -46,7 +46,7 @@
 #define __MAC_10_14 101400
 #endif
 
-#define APPLE_PROFILING_SIGNPOST ((!PLATFORM_MAC && __IPHONE_OS_VERSION_MIN_REQUIRED >= __IPHONE_12_0) || (PLATFORM_MAC && __MAC_OS_X_VERSION_MIN_REQUIRED >= __MAC_10_14))
+#define APPLE_PROFILING_SIGNPOST ((!PLATFORM_MAC && __IPHONE_OS_VERSION_MIN_REQUIRED >= __IPHONE_12_0) || PLATFORM_TVOS || (PLATFORM_MAC && __MAC_OS_X_VERSION_MIN_REQUIRED >= __MAC_10_14))
 
 #if APPLE_PROFILING_SIGNPOST
 #if PLATFORM_MAC	
