Save Game User Settings in the cloud for iOS / tvOS

diff --git a/Engine/Plugins/Online/IOS/OnlineSubsystemIOS/Source/Private/OnlineUserCloudInterfaceIOS.cpp b/Engine/Plugins/Online/IOS/OnlineSubsystemIOS/Source/Private/OnlineUserCloudInterfaceIOS.cpp
--- a/Engine/Plugins/Online/IOS/OnlineSubsystemIOS/Source/Private/OnlineUserCloudInterfaceIOS.cpp
+++ b/Engine/Plugins/Online/IOS/OnlineSubsystemIOS/Source/Private/OnlineUserCloudInterfaceIOS.cpp
@@ -3,7 +3,9 @@
 #include "OnlineUserCloudInterfaceIOS.h"
 #include "HAL/PlatformProcess.h"
 #include "Misc/ScopeLock.h"
+#include "Misc/FileHelper.h"
 #include "PlatformFeatures.h"
+#include "GameFramework/GameUserSettings.h"
 
 @implementation IOSCloudStorage
 
@@ -480,9 +482,9 @@ void FOnlineUserCloudInterfaceIOS::InitCloudSave(bool InIOSAlwaysSyncCloudFiles)
 	OnDeleteUserCloudFileCompleteDelegate = FOnDeleteUserFileCompleteDelegate::CreateRaw(this, &FOnlineUserCloudInterfaceIOS::OnDeleteUserCloudFileComplete);
 
 	// link delegates to the iOS save system
-	SaveSystem->OnWriteUserCloudFileBeginDelegate = FIOSSaveGameSystem::FOnWriteUserCloudFileBegin::CreateRaw(this, &FOnlineUserCloudInterfaceIOS::OnWriteUserCloudFileBegin);
-	SaveSystem->OnReadUserCloudFileBeginDelegate = FIOSSaveGameSystem::FOnReadUserCloudFileBegin::CreateRaw(this, &FOnlineUserCloudInterfaceIOS::OnReadUserCloudFileBegin);
-	SaveSystem->OnDeleteUserCloudFileBeginDelegate = FIOSSaveGameSystem::FOnDeleteUserCloudFileBegin::CreateRaw(this, &FOnlineUserCloudInterfaceIOS::OnDeleteUserCloudFileBegin);
+	SaveSystem->OnUpdateCloudDataFromLocalSave = FIOSSaveGameSystem::FUpdateCloudDataFromLocalSave::CreateRaw(this, &FOnlineUserCloudInterfaceIOS::OnWriteUserCloudFileBegin);
+	SaveSystem->OnUpdateLocalSaveFileFromCloud = FIOSSaveGameSystem::FUpdateLocalSaveFileFromCloud::CreateRaw(this, &FOnlineUserCloudInterfaceIOS::OnReadUserCloudFileBegin);
+	SaveSystem->OnDeleteCloudData = FIOSSaveGameSystem::FDeleteCloudData::CreateRaw(this, &FOnlineUserCloudInterfaceIOS::OnDeleteUserCloudFileBegin);
 
 	UniqueNetId = FUniqueNetIdIOS::Create();
 
@@ -490,6 +492,25 @@ void FOnlineUserCloudInterfaceIOS::InitCloudSave(bool InIOSAlwaysSyncCloudFiles)
 	
 	OnEnumerateUserCloudFilesCompleteDelegateHandle = AddOnEnumerateUserFilesCompleteDelegate_Handle(OnEnumerateUserCloudFilesCompleteDelegate);
 	EnumerateUserFiles(*UniqueNetId);
+    
+    auto RegisterItems = [](auto classToCall)
+    {
+		UGameUserSettings* UserSettings;
+        UserSettings = GEngine->GetGameUserSettings();
+        check(UserSettings);
+
+    	UserSettings->OnUpdateCloudDataFromGameUserSettings = UGameUserSettings::FUpdateCloudDataFromGameUserSettings::CreateRaw(classToCall, &FOnlineUserCloudInterfaceIOS::OnWriteUserCloudFileBegin);
+        UserSettings->OnUpdateGameUserSettingsFileFromCloud = UGameUserSettings::FUpdateGameUserSettingsFileFromCloud::CreateRaw(classToCall, &FOnlineUserCloudInterfaceIOS::OnReadUserCloudFileBegin);
+    };
+
+    if (GEngine)
+    {
+        RegisterItems(this);
+    }
+    else
+    {
+		FCoreDelegates::OnPostEngineInit.AddLambda(RegisterItems, this);
+    }
 }
 
 void FOnlineUserCloudInterfaceIOS::OnEnumerateUserFilesComplete(bool bWasSuccessful, const FUniqueNetId & UserId)
@@ -546,10 +567,7 @@ void FOnlineUserCloudInterfaceIOS::OnReadUserCloudFileComplete(bool bWasSuccessf
 		// locally cache the server data
 		TArray<uint8> FileContents;
 		GetFileContents(UserId, FileName, FileContents);
-
-		check(SaveSystem);
-
-		SaveSystem->SaveGameNoCloud(*FileName, FileContents);
+        FFileHelper::SaveArrayToFile(FileContents, *FileName);
 		
 		//flag that we have the latest record
 		[UpdateDictionary setObject : [NSNumber numberWithBool : false]  forKey : FileName.GetNSString()];
@@ -561,10 +579,18 @@ void FOnlineUserCloudInterfaceIOS::OnDeleteUserCloudFileComplete(bool bWasSucces
 	ClearOnDeleteUserFileCompleteDelegate_Handle(OnDeleteUserCloudFileCompleteDelegateHandle);
 }
 
-void FOnlineUserCloudInterfaceIOS::OnWriteUserCloudFileBegin(const FString &  FileName, const TArray<uint8>& FileContents)
+bool FOnlineUserCloudInterfaceIOS::OnWriteUserCloudFileBegin(const FString &  FileName)
 {
-	OnWriteUserCloudFileCompleteDelegateHandle = AddOnWriteUserFileCompleteDelegate_Handle(OnWriteUserCloudFileCompleteDelegate);
-	WriteUserFile(*UniqueNetId, FileName, (TArray<uint8>&)FileContents);
+    OnWriteUserCloudFileCompleteDelegateHandle = AddOnWriteUserFileCompleteDelegate_Handle(OnWriteUserCloudFileCompleteDelegate);
+    
+    TArray<uint8> FileData;
+    
+    if (FFileHelper::LoadFileToArray(FileData, *FileName))
+    {
+        return WriteUserFile(*UniqueNetId, FileName, FileData);
+    }
+    
+    return false;
 }
 
 bool FOnlineUserCloudInterfaceIOS::ShouldFetchRecordFromCloud(const FString &  FileName)
@@ -582,7 +608,7 @@ bool FOnlineUserCloudInterfaceIOS::ShouldFetchRecordFromCloud(const FString &  F
 	return true;
 }
 
-void FOnlineUserCloudInterfaceIOS::OnReadUserCloudFileBegin(const FString &  FileName, TArray<uint8>& FileContents)
+bool FOnlineUserCloudInterfaceIOS::OnReadUserCloudFileBegin(const FString & FileName)
 {
 #if !PLATFORM_TVOS
 	if (ShouldFetchRecordFromCloud(FileName))
@@ -599,11 +625,10 @@ void FOnlineUserCloudInterfaceIOS::OnReadUserCloudFileBegin(const FString &  Fil
 			FPlatformProcess::Sleep(0.01f);
 		}
 
-		GetFileContents(*UniqueNetId, FileName, FileContents);
-
 		// clean up temporary data
 		ClearFile(*UniqueNetId, FileName);
 	}
+    return true;
 }
 
 void FOnlineUserCloudInterfaceIOS::OnDeleteUserCloudFileBegin(const FString &  FileName)
diff --git a/Engine/Plugins/Online/IOS/OnlineSubsystemIOS/Source/Private/OnlineUserCloudInterfaceIOS.h b/Engine/Plugins/Online/IOS/OnlineSubsystemIOS/Source/Private/OnlineUserCloudInterfaceIOS.h
--- a/Engine/Plugins/Online/IOS/OnlineSubsystemIOS/Source/Private/OnlineUserCloudInterfaceIOS.h
+++ b/Engine/Plugins/Online/IOS/OnlineSubsystemIOS/Source/Private/OnlineUserCloudInterfaceIOS.h
@@ -143,13 +143,13 @@ private:
 	 *	Delegate in the iOS file save system
 	 * Called in SaveGame
 	 */
-	void OnWriteUserCloudFileBegin(const FString&  FileName, const TArray<uint8>& FileContents);
+	bool OnWriteUserCloudFileBegin(const FString&  FileName);
 
 	/**
 	 *	Delegate in the iOS file save system
 	 * Called in ReadGame
 	 */
-	void OnReadUserCloudFileBegin(const FString&  FileName, TArray<uint8>& FileContents);
+	bool OnReadUserCloudFileBegin(const FString & FileName);
 
 	/**
 	 *	Delegate in the iOS file save system
diff --git a/Engine/Source/Runtime/Engine/Classes/GameFramework/GameUserSettings.h b/Engine/Source/Runtime/Engine/Classes/GameFramework/GameUserSettings.h
--- a/Engine/Source/Runtime/Engine/Classes/GameFramework/GameUserSettings.h
+++ b/Engine/Source/Runtime/Engine/Classes/GameFramework/GameUserSettings.h
@@ -520,6 +520,19 @@ public:
 	{
 		OnGameUserSettingsUINeedsUpdate.Broadcast();
 	}
+    
+    /**
+     * We call this to refresh the config file from a listening Cloud subsystem, before we open it.
+     */
+    DECLARE_DELEGATE_RetVal_OneParam(bool, FUpdateGameUserSettingsFileFromCloud, const FString&);
+    FUpdateGameUserSettingsFileFromCloud OnUpdateGameUserSettingsFileFromCloud;
+        
+    /**
+     * We call this to notify any listening Cloud subsystem that we have updated the config file.
+     */
+    DECLARE_DELEGATE_RetVal_OneParam(bool, FUpdateCloudDataFromGameUserSettings, const FString&);
+    FUpdateCloudDataFromGameUserSettings OnUpdateCloudDataFromGameUserSettings;
+
 
 protected:
 	/**
diff --git a/Engine/Source/Runtime/Engine/Private/GameUserSettings.cpp b/Engine/Source/Runtime/Engine/Private/GameUserSettings.cpp
--- a/Engine/Source/Runtime/Engine/Private/GameUserSettings.cpp
+++ b/Engine/Source/Runtime/Engine/Private/GameUserSettings.cpp
@@ -3,11 +3,13 @@
 #include "GameFramework/GameUserSettings.h"
 #include "HAL/FileManager.h"
 #include "Misc/ConfigCacheIni.h"
+#include "Misc/FileHelper.h"
 #include "HAL/IConsoleManager.h"
 #include "GenericPlatform/GenericApplication.h"
 #include "Misc/App.h"
 #include "EngineGlobals.h"
 #include "Kismet/KismetSystemLibrary.h"
+#include "Kismet/GameplayStatics.h"
 #include "Engine/GameViewportClient.h"
 #include "UnrealEngine.h"
 #include "Framework/Application/SlateApplication.h"
@@ -515,6 +517,17 @@ void UGameUserSettings::LoadSettings(bool bForceReload/*=false*/)
 
 	if (bForceReload)
 	{
+        if (OnUpdateGameUserSettingsFileFromCloud.IsBound())
+        {
+            FString IniFileLocation = FPaths::GeneratedConfigDir() + UGameplayStatics::GetPlatformName() + "/" +  GGameUserSettingsIni + ".ini";
+            UE_LOG(LogTemp, Verbose, TEXT("%s"), *IniFileLocation);
+
+            if (!OnUpdateGameUserSettingsFileFromCloud.Execute(FString(*IniFileLocation)))
+            {
+                UE_LOG(LogTemp, Warning, TEXT("Failed to read the ini file from the Cloud interface %s"), *IniFileLocation);
+            }
+        }
+        
 		LoadConfigIni(bForceReload);
 	}
 	LoadConfig(GetClass(), *GGameUserSettingsIni);
@@ -551,6 +564,20 @@ void UGameUserSettings::SaveSettings()
 	// Save the Scalability state to the same ini file as it was loaded from in FEngineLoop::Preinit
 	Scalability::SaveState(GIsEditor ? GEditorSettingsIni : GGameUserSettingsIni);
 	SaveConfig(CPF_Config, *GGameUserSettingsIni);
+    
+    if (OnUpdateCloudDataFromGameUserSettings.IsBound())
+    {
+        FString IniFileLocation = FPaths::GeneratedConfigDir() + UGameplayStatics::GetPlatformName() + "/" +  GGameUserSettingsIni + ".ini";
+        UE_LOG(LogTemp, Verbose, TEXT("%s"), *IniFileLocation);
+
+        bool bDidSucceed = false;
+        bDidSucceed = OnUpdateCloudDataFromGameUserSettings.Execute(FString(*IniFileLocation));
+        
+        if (!bDidSucceed)
+        {
+            UE_LOG(LogTemp, Warning, TEXT("Failed to load the ini file from the Cloud interface %s"), *IniFileLocation);
+        }
+    }
 }
 
 void UGameUserSettings::LoadConfigIni(bool bForceReload/*=false*/)
@@ -962,4 +989,4 @@ int32 UGameUserSettings::GetCurrentHDRDisplayNits() const
 bool UGameUserSettings::IsHDREnabled() const
 {
 	return bUseHDRDisplayOutput;
-}
\ No newline at end of file
+}
diff --git a/Engine/Source/Runtime/IOS/IOSPlatformFeatures/Private/IOSSaveGameSystem.cpp b/Engine/Source/Runtime/IOS/IOSPlatformFeatures/Private/IOSSaveGameSystem.cpp
--- a/Engine/Source/Runtime/IOS/IOSPlatformFeatures/Private/IOSSaveGameSystem.cpp
+++ b/Engine/Source/Runtime/IOS/IOSPlatformFeatures/Private/IOSSaveGameSystem.cpp
@@ -50,38 +50,34 @@ ISaveGameSystem::ESaveExistsResult FIOSSaveGameSystem::DoesSaveGameExistWithResu
 	return ESaveExistsResult::DoesNotExist;
 }
 
-bool FIOSSaveGameSystem::SaveGameNoCloud(const TCHAR* Name, const TArray<uint8>& Data)
-{
-	return FFileHelper::SaveArrayToFile(Data, *GetSaveGamePath(Name));
-}
-
 bool FIOSSaveGameSystem::SaveGame(bool bAttemptToUseUI, const TCHAR* Name, const int32 UserIndex, const TArray<uint8>& Data)
 {
-	// send to the iCloud, if enabled
-	OnWriteUserCloudFileBeginDelegate.ExecuteIfBound(FString(Name), Data);
-
-	return FFileHelper::SaveArrayToFile(Data, *GetSaveGamePath(Name));
+    bool DidSucceed = false;
+    
+    if (OnUpdateCloudDataFromLocalSave.IsBound())
+    {
+        DidSucceed = OnUpdateCloudDataFromLocalSave.Execute(FString::Printf(TEXT("%s""SaveGames/%s.sav"), *FPaths::ProjectSavedDir(), Name));
+    }
+
+    return DidSucceed ? DidSucceed : FFileHelper::SaveArrayToFile(Data, *GetSaveGamePath(Name));
 }
 
 bool FIOSSaveGameSystem::LoadGame(bool bAttemptToUseUI, const TCHAR* Name, const int32 UserIndex, TArray<uint8>& Data)
 {
-	// try to read it from the iCloud, if enabled
-	OnReadUserCloudFileBeginDelegate.ExecuteIfBound(FString(Name), Data);
-
-	if (Data.Num() > 0)
-	{
-		// we've received data from the iCloud, the save file was overwritten
-		return true;
-	}
-
-	// no iCloud data, read from the local storage
-	return FFileHelper::LoadFileToArray(Data, *GetSaveGamePath(Name));
+    bool DidSucceed = false;
+
+    if (OnUpdateLocalSaveFileFromCloud.IsBound())
+    {
+        DidSucceed = OnUpdateLocalSaveFileFromCloud.Execute(FString::Printf(TEXT("%s""SaveGames/%s.sav"), *FPaths::ProjectSavedDir(), Name));
+    }
+    
+    return DidSucceed ? DidSucceed : FFileHelper::LoadFileToArray(Data, *GetSaveGamePath(Name));
 }
 
 bool FIOSSaveGameSystem::DeleteGame(bool bAttemptToUseUI, const TCHAR* Name, const int32 UserIndex)
 {
 	// delete the file from the iCloud
-	OnDeleteUserCloudFileBeginDelegate.ExecuteIfBound(FString(Name));
+    OnDeleteCloudData.ExecuteIfBound(FString(Name));
 
 	// delete the file from the local storage
 	return IFileManager::Get().Delete(*GetSaveGamePath(Name), true, false, !bAttemptToUseUI);
diff --git a/Engine/Source/Runtime/IOS/IOSPlatformFeatures/Public/IOSSaveGameSystem.h b/Engine/Source/Runtime/IOS/IOSPlatformFeatures/Public/IOSSaveGameSystem.h
--- a/Engine/Source/Runtime/IOS/IOSPlatformFeatures/Public/IOSSaveGameSystem.h
+++ b/Engine/Source/Runtime/IOS/IOSPlatformFeatures/Public/IOSSaveGameSystem.h
@@ -29,35 +29,27 @@ public:
 	}
 
 	virtual ESaveExistsResult DoesSaveGameExistWithResult(const TCHAR* Name, const int32 UserIndex) override;
-
-	/**
-	 * Called on the initial iCloud sync
-	 */
-	virtual bool SaveGameNoCloud(const TCHAR* Name, const TArray<uint8>& Data);
 	
 	virtual bool SaveGame(bool bAttemptToUseUI, const TCHAR* Name, const int32 UserIndex, const TArray<uint8>& Data) override;
-	/**
-	 * Called when writing the savegame file
-	 * send the file to iCloud, if enabled
-	 */
-	DECLARE_DELEGATE_TwoParams(FOnWriteUserCloudFileBegin, const FString&, const TArray<uint8>&);
-	FOnWriteUserCloudFileBegin OnWriteUserCloudFileBeginDelegate;
+	 /**
+     * We call this to notify any listening Cloud subsystem that we have updated a save file.
+     */
+    DECLARE_DELEGATE_RetVal_OneParam(bool, FUpdateCloudDataFromLocalSave, const FString&);
+    FUpdateCloudDataFromLocalSave OnUpdateCloudDataFromLocalSave;
 
 	virtual bool LoadGame(bool bAttemptToUseUI, const TCHAR* Name, const int32 UserIndex, TArray<uint8>& Data) override;
 	/**
-	 * Called when reading the savegame file
-	 * read the file from iCloud, if enabled
-	 */
-	DECLARE_DELEGATE_TwoParams(FOnReadUserCloudFileBegin, const FString&, TArray<uint8>&);
-	FOnReadUserCloudFileBegin OnReadUserCloudFileBeginDelegate;
+     * We call this to refresh a save file from a listening Cloud subsystem, before we open it.
+     */
+    DECLARE_DELEGATE_RetVal_OneParam(bool, FUpdateLocalSaveFileFromCloud, const FString&);
+    FUpdateLocalSaveFileFromCloud OnUpdateLocalSaveFileFromCloud;
 
 	virtual bool DeleteGame(bool bAttemptToUseUI, const TCHAR* Name, const int32 UserIndex) override;
 	/**
-	 * Called when deleting the savegame file
-	 * delete the file from iCloud, if enabled
-	 */
-	DECLARE_DELEGATE_OneParam(FOnDeleteUserCloudFileBegin, const FString&);
-	FOnDeleteUserCloudFileBegin OnDeleteUserCloudFileBeginDelegate;
+     * We call this to delete a save file from a listening Cloud subsystem.
+     */
+	DECLARE_DELEGATE_OneParam(FDeleteCloudData, const FString&);
+	FDeleteCloudData OnDeleteCloudData;
 
 private:
 	/**
