Fix non-unity errors
#rnx

diff --git a/Engine/Plugins/Online/IOS/OnlineSubsystemIOS/Source/Private/OnlineUserCloudInterfaceIOS.cpp b/Engine/Plugins/Online/IOS/OnlineSubsystemIOS/Source/Private/OnlineUserCloudInterfaceIOS.cpp
--- a/Engine/Plugins/Online/IOS/OnlineSubsystemIOS/Source/Private/OnlineUserCloudInterfaceIOS.cpp
+++ b/Engine/Plugins/Online/IOS/OnlineSubsystemIOS/Source/Private/OnlineUserCloudInterfaceIOS.cpp
@@ -1,9 +1,11 @@
 // Copyright Epic Games, Inc. All Rights Reserved.
 
 #include "OnlineUserCloudInterfaceIOS.h"
+#include "Engine/Engine.h"
 #include "HAL/PlatformProcess.h"
-#include "Misc/ScopeLock.h"
+#include "Misc/CoreDelegates.h"
 #include "Misc/FileHelper.h"
+#include "Misc/ScopeLock.h"
 #include "PlatformFeatures.h"
 #include "GameFramework/GameUserSettings.h"
 
diff --git a/Engine/Plugins/VirtualProduction/RemoteControl/Source/RemoteControlUI/Private/UI/SRCPanelExposedEntitiesList.h b/Engine/Plugins/VirtualProduction/RemoteControl/Source/RemoteControlUI/Private/UI/SRCPanelExposedEntitiesList.h
--- a/Engine/Plugins/VirtualProduction/RemoteControl/Source/RemoteControlUI/Private/UI/SRCPanelExposedEntitiesList.h
+++ b/Engine/Plugins/VirtualProduction/RemoteControl/Source/RemoteControlUI/Private/UI/SRCPanelExposedEntitiesList.h
@@ -12,7 +12,6 @@
 
 struct FRCPanelGroup;
 class FRCPanelWidgetRegistry;
-struct FRemoteControlTarget;
 class FDragDropOperation;
 struct FGuid;
 class FReply;
diff --git a/Engine/Source/Runtime/Core/Public/Tasks/Task.h b/Engine/Source/Runtime/Core/Public/Tasks/Task.h
--- a/Engine/Source/Runtime/Core/Public/Tasks/Task.h
+++ b/Engine/Source/Runtime/Core/Public/Tasks/Task.h
@@ -6,6 +6,7 @@
 #include "Async/Fundamental/Task.h"
 #include "Containers/StaticArray.h"
 #include "HAL/Event.h"
+#include "HAL/IConsoleManager.h"
 #include "CoreTypes.h"
 
 namespace UE::Tasks
