Fix iOS leaderboards crashing UnrealMatch3 when not available.  Also, remove some pre-iOS 13 APIs and assume iOS14+

diff --git a/Engine/Plugins/Online/IOS/OnlineSubsystemIOS/Source/Private/OnlineLeaderboardsInterfaceIOS.cpp b/Engine/Plugins/Online/IOS/OnlineSubsystemIOS/Source/Private/OnlineLeaderboardsInterfaceIOS.cpp
--- a/Engine/Plugins/Online/IOS/OnlineSubsystemIOS/Source/Private/OnlineLeaderboardsInterfaceIOS.cpp
+++ b/Engine/Plugins/Online/IOS/OnlineSubsystemIOS/Source/Private/OnlineLeaderboardsInterfaceIOS.cpp
@@ -48,20 +48,18 @@
 bool FOnlineLeaderboardsIOS::ReadLeaderboardCompletionDelegate(NSArray* players, FOnlineLeaderboardReadRef& InReadObject)
 {
     auto ReadObject = InReadObject;
-    bool bTriggeredReadRequest = false;
+	bool bTriggeredReadRequest = true;
 
     CachedLeaderboard = [GKLeaderboard alloc];
-	if ([GKLeaderboard instancesRespondToSelector:@selector(initWithPlayers:)] == YES)
-	{
         [CachedLeaderboard loadEntriesForPlayers:players timeScope:GKLeaderboardTimeScopeAllTime completionHandler: ^(GKLeaderboardEntry *entries, NSArray<GKLeaderboardEntry *> *scores, NSError *Error)
          {
-            //GKLeaderboardEntry *entriesCopy = entries;
             bool bWasSuccessful = (Error == nil) && [scores count] > 0;
                         
+			UE_LOG_ONLINE_LEADERBOARD(Display, TEXT("FOnlineLeaderboardsIOS::loadScoresWithCompletionHandler() - %s"), (bWasSuccessful ? TEXT("Success!") : TEXT("Failed!, no scores retrieved")));
+		
             if (bWasSuccessful)
             {
                 bWasSuccessful = [scores count] > 0;
-                UE_LOG_ONLINE_LEADERBOARD(Display, TEXT("FOnlineLeaderboardsIOS::loadScoresWithCompletionHandler() - %s"), (bWasSuccessful ? TEXT("Success!") : TEXT("Failed!, no scores retrieved")));
                 for (GKLeaderboardEntry *entry in scores)
                 {
                     FString PlayerIDString;
@@ -110,6 +108,8 @@
             }
             else if (Error)
             {
+				LeaderboardPlayer = nil;
+			
                 // if we have failed to read the leaderboard then report this
                 NSDictionary *userInfo = [Error userInfo];
                 NSString *errstr = [[userInfo objectForKey : NSUnderlyingErrorKey] localizedDescription];
@@ -125,15 +125,7 @@
                 return true;
              }];
         }];
-    }
     
-    // If we have failed to kick off a read request, we should still tell whoever is listening.
-    if (!bTriggeredReadRequest)
-    {
-        UE_LOG_ONLINE_LEADERBOARD(Display, TEXT("FOnlineLeaderboardsIOS::loadScoresWithCompletionHandler() - Failed!"));
-        TriggerOnLeaderboardReadCompleteDelegates(false);
-    }
-    
     return bTriggeredReadRequest;
 }
 
@@ -166,23 +158,20 @@
 		}
 
 		// Kick off a game center read request for the list of users
-		if ([GKLeaderboard instancesRespondToSelector:@selector(loadEntriesForPlayerScope:timeScope:range:completionHandler:)] == YES)
-		{
-            [GKPlayer loadPlayersForIdentifiers:FriendIds withCompletionHandler:^(NSArray *players, NSError *Error)
-             {
-                bool bWasSuccessful = (Error == nil) && [players count] > 0;
-             
-                if (bWasSuccessful)
-                {
-                    bWasSuccessful = [players count] > 0;
-					ReadLeaderboardCompletionDelegate(players, ReadObject);
-				}
-             }];
-        }
-        else
-        {
-			return false;
-        }
+		[GKPlayer loadPlayersForIdentifiers:FriendIds withCompletionHandler:^(NSArray *players, NSError *Error)
+		 {
+			bool bWasSuccessful = (Error == nil) && [players count] > 0;
+			if (bWasSuccessful)
+			{
+				bWasSuccessful = [players count] > 0;
+			}
+			// even if not successful, need to call the delegate to initialize Leaderboard and LocalPlayer
+			ReadLeaderboardCompletionDelegate(players, ReadObject);
+		}];
+	}
+	else
+	{
+		return false;
 	}
 
 	return true;
@@ -296,45 +285,47 @@
         const int32 UnreportedScoreCount = UnreportedScores.count;
         bBeganFlushingScores = UnreportedScoreCount > 0;
         
-        if (bBeganFlushingScores)
-        {
-            NSMutableArray *ArrayCopy = [[NSMutableArray alloc] initWithArray:UnreportedScores];
-            
-            [UnreportedScores release];
-            UnreportedScores = nil;
-            
-            dispatch_async(dispatch_get_main_queue(), ^
-                           {
-                for (NSNumber* scoreReport in ArrayCopy)
-                {
-                    NSInteger ScoreReportInt = [scoreReport integerValue];
-                    
-                    [CachedLeaderboard submitScore:ScoreReportInt context:0 player:LeaderboardPlayer completionHandler: ^ (NSError *error)
-                     {
-                        // Tell whoever was listening that we have written (or failed to write) to the leaderboard
-                        bool bSucceeded = error == NULL;
-                        if (bSucceeded)
-                        {
-                            UE_LOG_ONLINE_LEADERBOARD(Display, TEXT("Flushed %d scores to Game Center"), UnreportedScoreCount);
-                        }
-                        else
-                        {
-                            UE_LOG_ONLINE_LEADERBOARD(Display, TEXT("Error while flushing scores (code %d)"), [error code]);
-                        }
-                        
-                        [ArrayCopy release];
-                        
-                        // Report back to the game thread whether this succeeded.
-                        [FIOSAsyncTask CreateTaskWithBlock : ^ bool(void)
-                         {
-                            TriggerOnLeaderboardFlushCompleteDelegates(SessionName, bSucceeded);
-                            return true;
-                        }];
-                    }];
-                }
-            });
-        }
-        
+		if (bBeganFlushingScores)
+		{
+			NSMutableArray *ArrayCopy = [[NSMutableArray alloc] initWithArray:UnreportedScores];
+			
+			[UnreportedScores release];
+			UnreportedScores = nil;
+			
+			dispatch_async(dispatch_get_main_queue(),
+			^{
+				if (LeaderboardPlayer != nil)
+				{
+					for (NSNumber* scoreReport in ArrayCopy)
+					{
+						NSInteger ScoreReportInt = [scoreReport integerValue];
+						
+						[CachedLeaderboard submitScore:ScoreReportInt context:0 player:LeaderboardPlayer completionHandler: ^ (NSError *error)
+						 {
+							// Tell whoever was listening that we have written (or failed to write) to the leaderboard
+							bool bSucceeded = error == NULL;
+							if (bSucceeded)
+							{
+								UE_LOG_ONLINE_LEADERBOARD(Display, TEXT("Flushed %d scores to Game Center"), UnreportedScoreCount);
+							}
+							else
+							{
+								UE_LOG_ONLINE_LEADERBOARD(Display, TEXT("Error while flushing scores (code %d)"), [error code]);
+							}
+							
+							// Report back to the game thread whether this succeeded.
+							[FIOSAsyncTask CreateTaskWithBlock : ^ bool(void)
+							 {
+								TriggerOnLeaderboardFlushCompleteDelegates(SessionName, bSucceeded);
+								return true;
+							}];
+						}];
+					}
+				}
+				[ArrayCopy release];
+			});
+		}
+	
         // If we didn't begin writing to the leaderboard we should still notify whoever was listening.
         if (!bBeganFlushingScores)
         {
