Fix crash in Match3 when doing the first flush of leaderboard on iOS

diff --git a/Engine/Plugins/Online/IOS/OnlineSubsystemIOS/Source/Private/OnlineLeaderboardsInterfaceIOS.cpp b/Engine/Plugins/Online/IOS/OnlineSubsystemIOS/Source/Private/OnlineLeaderboardsInterfaceIOS.cpp
--- a/Engine/Plugins/Online/IOS/OnlineSubsystemIOS/Source/Private/OnlineLeaderboardsInterfaceIOS.cpp
+++ b/Engine/Plugins/Online/IOS/OnlineSubsystemIOS/Source/Private/OnlineLeaderboardsInterfaceIOS.cpp
@@ -53,6 +53,8 @@
     CachedLeaderboard = [GKLeaderboard alloc];
         [CachedLeaderboard loadEntriesForPlayers:players timeScope:GKLeaderboardTimeScopeAllTime completionHandler: ^(GKLeaderboardEntry *entries, NSArray<GKLeaderboardEntry *> *scores, NSError *Error)
          {
+            bReadLeaderboardFinished = true;
+
             bool bWasSuccessful = (Error == nil) && [scores count] > 0;
                         
 			UE_LOG_ONLINE_LEADERBOARD(Display, TEXT("FOnlineLeaderboardsIOS::loadScoresWithCompletionHandler() - %s"), (bWasSuccessful ? TEXT("Success!") : TEXT("Failed!, no scores retrieved")));
@@ -135,6 +137,7 @@
 
 	UE_LOG_ONLINE_LEADERBOARD(Display, TEXT("FOnlineLeaderboardsIOS::ReadLeaderboards()"));
 
+    bReadLeaderboardFinished = false;
 	ReadObject->ReadState = EOnlineAsyncTaskState::Failed;
 	ReadObject->Rows.Empty();
 	
@@ -278,6 +281,20 @@
 bool FOnlineLeaderboardsIOS::FlushLeaderboards(const FName& SessionName)
 {
     UE_LOG_ONLINE_LEADERBOARD(Display, TEXT("FOnlineLeaderboardsIOS::FlushLeaderboards()"));
+    
+    
+    int SleepLimit = 30;
+    while(!bReadLeaderboardFinished && SleepLimit != 0)
+    {
+        usleep(100000);
+        --SleepLimit;
+    }
+    if (SleepLimit == 0)
+    {
+        UE_LOG_ONLINE_LEADERBOARD(Warning, TEXT("Leaderboard Could not be flusshed"));
+        return false;
+    }
+    
     bool bBeganFlushingScores = false;
     
     if ((IdentityInterface->GetLocalGameCenterUser() != NULL) && IdentityInterface->GetLocalGameCenterUser().isAuthenticated)
@@ -292,10 +309,10 @@
 			[UnreportedScores release];
 			UnreportedScores = nil;
 			
+            if (LeaderboardPlayer != nil)
+            {
 			dispatch_async(dispatch_get_main_queue(),
 			^{
-				if (LeaderboardPlayer != nil)
-				{
 					for (NSNumber* scoreReport in ArrayCopy)
 					{
 						NSInteger ScoreReportInt = [scoreReport integerValue];
@@ -321,9 +338,13 @@
 							}];
 						}];
 					}
+                [ArrayCopy release];
+                });
 				}
+            else
+            {
 				[ArrayCopy release];
-			});
+            }
 		}
 	
         // If we didn't begin writing to the leaderboard we should still notify whoever was listening.
diff --git a/Engine/Plugins/Online/IOS/OnlineSubsystemIOS/Source/Private/OnlineLeaderboardsInterfaceIOS.h b/Engine/Plugins/Online/IOS/OnlineSubsystemIOS/Source/Private/OnlineLeaderboardsInterfaceIOS.h
--- a/Engine/Plugins/Online/IOS/OnlineSubsystemIOS/Source/Private/OnlineLeaderboardsInterfaceIOS.h
+++ b/Engine/Plugins/Online/IOS/OnlineSubsystemIOS/Source/Private/OnlineLeaderboardsInterfaceIOS.h
@@ -18,6 +18,7 @@ private:
 	NSMutableArray *UnreportedScores;
     GKPlayer *LeaderboardPlayer;
     GKLeaderboard *CachedLeaderboard;
+	bool bReadLeaderboardFinished;
     
     bool ReadLeaderboardCompletionDelegate(NSArray* players, FOnlineLeaderboardReadRef& ReadObject);
     
