Fix iOS build

diff --git a/Engine/Plugins/Online/IOS/OnlineSubsystemIOS/Source/Private/OnlineStoreKitHelper.cpp b/Engine/Plugins/Online/IOS/OnlineSubsystemIOS/Source/Private/OnlineStoreKitHelper.cpp
--- a/Engine/Plugins/Online/IOS/OnlineSubsystemIOS/Source/Private/OnlineStoreKitHelper.cpp
+++ b/Engine/Plugins/Online/IOS/OnlineSubsystemIOS/Source/Private/OnlineStoreKitHelper.cpp
@@ -95,25 +95,24 @@ FStoreKitTransactionData::FStoreKitTransactionData(const SKPaymentTransaction* T
 }
 
 ////////////////////////////////////////////////////////////////////
-/// FStoreKitHelper implementation
+/// FStoreKitHelperV2 implementation
 
-@implementation FStoreKitHelper
+@implementation FStoreKitHelperV2
 @synthesize Request;
 @synthesize AvailableProducts;
 
 - (id)init
 {
-	self = [super init];
-
 	[FPaymentTransactionObserver sharedInstance].eventReceivedDelegate = self;
-
+    self.PendingTransactions = [NSMutableSet setWithCapacity:5];
 	return self;
 }
 
 -(void)dealloc
 {
 	[FPaymentTransactionObserver sharedInstance].eventReceivedDelegate = nil;
-
+    [_PendingTransactions release];
+    
 	[Request release];
 	[AvailableProducts release];
 	[super dealloc];
@@ -128,16 +127,6 @@ FStoreKitTransactionData::FStoreKitTransactionData(const SKPaymentTransaction* T
     UE_LOG_ONLINE_STOREV2(Log, TEXT("FStoreKitHelper::removedTransaction"));
 }
 
--(void)purchaseInProgress: (SKPaymentTransaction *)transaction
-{
-    UE_LOG_ONLINE_STOREV2(Log, TEXT("FStoreKitHelper::purchaseInProgress"));
-}
-
--(void)purchaseDeferred: (SKPaymentTransaction *)transaction
-{
-    UE_LOG_ONLINE_STOREV2(Log, TEXT("FStoreKitHelper::purchaseDeferred"));
-}
-
 -(void)requestProductData: (NSMutableSet*)productIDs
 {
     UE_LOG_ONLINE_STOREV2(Log, TEXT("FStoreKitHelper::requestProductData"));
@@ -158,23 +147,6 @@ FStoreKitTransactionData::FStoreKitTransactionData(const SKPaymentTransaction* T
 	[Request start];
 }
 
--(void)productsRequest: (SKProductsRequest *)request didReceiveResponse : (SKProductsResponse *)response
-{
-    UE_LOG_ONLINE_STOREV2(Log, TEXT("FStoreKitHelper::didReceiveResponse"));
-	// Direct the response back to the store interface
-	[FIOSAsyncTask CreateTaskWithBlock : ^ bool(void)
-	{
-		IOnlineSubsystem* OnlineSub = IOnlineSubsystem::Get(IOS_SUBSYSTEM);
-		if (ensure(OnlineSub))
-		{
-		}
-		
-		return true;
-	}];
-
-	[request autorelease];
-}
-
 -(void)requestDidFinish:(SKRequest*)request
 {
 	if ([Request isKindOfClass : [SKReceiptRefreshRequest class]])
@@ -190,25 +162,6 @@ FStoreKitTransactionData::FStoreKitTransactionData(const SKPaymentTransaction* T
 	[Request start];
 }
 
-@end
-
-////////////////////////////////////////////////////////////////////
-/// FStoreKitHelperV2 implementation
-
-@implementation FStoreKitHelperV2
-
-- (id)init
-{
-	self = [super init];
-	self.PendingTransactions = [NSMutableSet setWithCapacity:5];
-	return self;
-}
--(void)dealloc
-{
-	[_PendingTransactions release];
-	[super dealloc];
-}
-
 -(void)restoreCompletedTransactionsFinished
 {
 	UE_LOG_ONLINE_STOREV2(Log, TEXT("FStoreKitHelperV2::restoreCompletedTransactionsFinished"));
diff --git a/Engine/Plugins/Online/IOS/OnlineSubsystemIOS/Source/Private/OnlineStoreKitHelper.h b/Engine/Plugins/Online/IOS/OnlineSubsystemIOS/Source/Private/OnlineStoreKitHelper.h
--- a/Engine/Plugins/Online/IOS/OnlineSubsystemIOS/Source/Private/OnlineStoreKitHelper.h
+++ b/Engine/Plugins/Online/IOS/OnlineSubsystemIOS/Source/Private/OnlineStoreKitHelper.h
@@ -129,26 +129,11 @@ typedef FOnProductsRequestResponse::FDelegate FOnProductsRequestResponseDelegate
 
 @end
 
-/** Store kit request object, holds information about the products we are purchasing, or querying. */
-@property (nonatomic, strong) SKRequest *Request;
-/** collection of available products attaced through a store kit request */
-@property (nonatomic, strong) NSArray *AvailableProducts;
-
-/** Helper fn to start a store kit purchase request */
--(void)makePurchase:(NSMutableSet*)productIDs;
-
-/** Helper fn to start a store kit purchase information query request */
--(void)requestProductData:(NSMutableSet*)productIDs;
-
-/** Helper fn to restore previously purchased products */
--(void)restorePurchases;
-@end
-
 /**
  * Helper class, which allows us to manage IAP product information requests, AND transactions
  * (version used by OnlineStoreIOS.h and OnlinePurchaseIOS.h)
  */
-@interface FStoreKitHelperV2 : FStoreKitHelper
+@interface FStoreKitHelperV2 : NSObject<SKProductsRequestDelegate, SKRequestDelegate, FPaymentTransactionObserverEventReceivedDelegate>
 {
 	/** delegate fired when a product request completes */
 	FOnProductsRequestResponse _OnProductRequestResponse;
@@ -167,6 +152,11 @@ typedef FOnProductsRequestResponse::FDelegate FOnProductsRequestResponseDelegate
 /** list of known pending transactions that are in the queue but not marked as redeemed by the app */
 @property (nonatomic, strong) NSMutableSet* PendingTransactions;
 
+/** Store kit request object, holds information about the products we are purchasing, or querying. */
+@property (nonatomic, strong) SKRequest *Request;
+/** collection of available products attaced through a store kit request */
+@property (nonatomic, strong) NSArray *AvailableProducts;
+
 -(FDelegateHandle)AddOnProductRequestResponse: (const FOnProductsRequestResponseDelegate&) Delegate;
 -(FDelegateHandle)AddOnTransactionComplete: (const FOnTransactionCompleteIOSDelegate&) Delegate;
 -(FDelegateHandle)AddOnTransactionRestored: (const FOnTransactionRestoredIOSDelegate&) Delegate;
@@ -196,6 +186,16 @@ typedef FOnProductsRequestResponse::FDelegate FOnProductsRequestResponseDelegate
  */
 -(void)makePurchase:(NSArray*)products WithUserId: (const FString&) userId SimulateAskToBuy: (bool) bAskToBuy;
 
+/** Helper fn to start a store kit purchase request */
+-(void)makePurchase:(NSMutableSet*)productIDs;
+
+/** Helper fn to start a store kit purchase information query request */
+-(void)requestProductData:(NSMutableSet*)productIDs;
+
+/** Helper fn to restore previously purchased products */
+-(void)restorePurchases;
+
+
 /**
  * Make a request for product information
  *
