Added back code removed accidentally on CL20831323
Fix wrong key for map

diff --git a/Engine/Plugins/Online/IOS/OnlineSubsystemIOS/Source/Private/OnlinePurchaseIOS.cpp b/Engine/Plugins/Online/IOS/OnlineSubsystemIOS/Source/Private/OnlinePurchaseIOS.cpp
--- a/Engine/Plugins/Online/IOS/OnlineSubsystemIOS/Source/Private/OnlinePurchaseIOS.cpp
+++ b/Engine/Plugins/Online/IOS/OnlineSubsystemIOS/Source/Private/OnlinePurchaseIOS.cpp
@@ -223,7 +223,8 @@ void FOnlinePurchaseIOS::GetReceipts(const FUniqueNetId& UserId, TArray<FPurchas
 	OutReceipts.Empty();
 
 	// Add the cached list of user purchases
-	const TArray<TSharedRef<FPurchaseReceipt>>* UserCompletedTransactions = CompletedTransactions.Find(UserId.ToString());
+    const FString UserIdStr = IOSUSER;
+	const TArray<TSharedRef<FPurchaseReceipt>>* UserCompletedTransactions = CompletedTransactions.Find(UserIdStr);
 	if (UserCompletedTransactions != nullptr)
 	{
 		for (int32 Idx = 0; Idx < UserCompletedTransactions->Num(); Idx++)
diff --git a/Engine/Plugins/Online/IOS/OnlineSubsystemIOS/Source/Private/OnlineStoreKitHelper.cpp b/Engine/Plugins/Online/IOS/OnlineSubsystemIOS/Source/Private/OnlineStoreKitHelper.cpp
--- a/Engine/Plugins/Online/IOS/OnlineSubsystemIOS/Source/Private/OnlineStoreKitHelper.cpp
+++ b/Engine/Plugins/Online/IOS/OnlineSubsystemIOS/Source/Private/OnlineStoreKitHelper.cpp
@@ -120,8 +120,75 @@ FStoreKitTransactionData::FStoreKitTransactionData(const SKPaymentTransaction* T
 
 -(void)onPaymentTransactionObserverEventReceived
 {
+    [FIOSAsyncTask CreateTaskWithBlock : ^ bool(void)
+	{
+		[self pumpObserverEventQueue];
+
+		return true;
+	}];
+}
+
+-(void)pumpObserverEventQueue
+{
+	FPaymentTransactionObserver* Observer = [FPaymentTransactionObserver sharedInstance];
+
+	FPaymentTransactionObserverEvent ObserverEvent;
+	while ([Observer getEventQueue].Dequeue(ObserverEvent))
+	{
+		switch (ObserverEvent.Type)
+		{
+		case EPaymentTransactionObserverEventType::UpdatedTransaction:
+			[self updatedTransaction : ObserverEvent.Transaction];
+			break;
+		case EPaymentTransactionObserverEventType::RemovedTransaction:
+			[self removedTransaction : ObserverEvent.Transaction];
+			break;
+		case EPaymentTransactionObserverEventType::RestoreCompletedTransactionsFailed:
+			[self restoreCompletedTransactionsFailedWithError : ObserverEvent.ErrorCode];
+			break;
+		case EPaymentTransactionObserverEventType::RestoreCompletedTransactionsFinished:
+			[self restoreCompletedTransactionsFinished];
+			break;
+		default:
+			break;
+		}
+	}
 }
 
+-(void)updatedTransaction : (SKPaymentTransaction*)transaction
+{
+	UE_LOG_ONLINE_STOREV2(Log, TEXT("FStoreKitHelper::updatedTransaction"));
+
+	// Parse the generic transaction update into appropriate execution paths
+	switch ([transaction transactionState])
+	{
+	case SKPaymentTransactionStatePurchased:
+		if (FParse::Param(FCommandLine::Get(), TEXT("disableiosredeem")))
+		{
+			UE_LOG_ONLINE_STOREV2(Log, TEXT("FStoreKitHelper::completeTransaction (disabled)"));
+		}
+		else
+		{
+			[self completeTransaction : transaction] ;
+		}
+		break;
+	case SKPaymentTransactionStateFailed:
+		[self failedTransaction : transaction] ;
+		break;
+	case SKPaymentTransactionStateRestored:
+		[self restoreTransaction : transaction] ;
+		break;
+	case SKPaymentTransactionStatePurchasing:
+		[self purchaseInProgress : transaction] ;
+		break;
+	case SKPaymentTransactionStateDeferred:
+		[self purchaseDeferred : transaction] ;
+		break;
+	default:
+		UE_LOG_ONLINE_STOREV2(Warning, TEXT("FStoreKitHelper unhandled state: %d"), [transaction transactionState]);
+		break;
+	}}
+
 -(void)removedTransaction : (SKPaymentTransaction*)transaction
 {
     UE_LOG_ONLINE_STOREV2(Log, TEXT("FStoreKitHelper::removedTransaction"));
diff --git a/Engine/Plugins/Online/IOS/OnlineSubsystemIOS/Source/Private/OnlineStoreKitHelper.h b/Engine/Plugins/Online/IOS/OnlineSubsystemIOS/Source/Private/OnlineStoreKitHelper.h
--- a/Engine/Plugins/Online/IOS/OnlineSubsystemIOS/Source/Private/OnlineStoreKitHelper.h
+++ b/Engine/Plugins/Online/IOS/OnlineSubsystemIOS/Source/Private/OnlineStoreKitHelper.h
@@ -157,6 +157,9 @@ typedef FOnProductsRequestResponse::FDelegate FOnProductsRequestResponseDelegate
 /** collection of available products attaced through a store kit request */
 @property (nonatomic, strong) NSArray *AvailableProducts;
 
+/** Pump any events that are enqueued on the observer. */
+-(void)pumpObserverEventQueue;
+
 -(FDelegateHandle)AddOnProductRequestResponse: (const FOnProductsRequestResponseDelegate&) Delegate;
 -(FDelegateHandle)AddOnTransactionComplete: (const FOnTransactionCompleteIOSDelegate&) Delegate;
 -(FDelegateHandle)AddOnTransactionRestored: (const FOnTransactionRestoredIOSDelegate&) Delegate;
