Fix iOS15 deprecation warnings.

#rnx

diff --git a/Engine/Plugins/Online/IOS/OnlineSubsystemIOS/Source/Private/OnlineTurnBasedInterfaceIOS.cpp b/Engine/Plugins/Online/IOS/OnlineSubsystemIOS/Source/Private/OnlineTurnBasedInterfaceIOS.cpp
--- a/Engine/Plugins/Online/IOS/OnlineSubsystemIOS/Source/Private/OnlineTurnBasedInterfaceIOS.cpp
+++ b/Engine/Plugins/Online/IOS/OnlineSubsystemIOS/Source/Private/OnlineTurnBasedInterfaceIOS.cpp
@@ -8,6 +8,7 @@
 #include "OnlineAsyncTaskManager.h"
 #include "Net/RepLayout.h"
 #include "Interfaces/TurnBasedMatchInterface.h"
+#include <GameKit/GKLocalPlayer.h>
 
 DEFINE_LOG_CATEGORY_STATIC(LogTurnBasedInterfaceIOS, Verbose, All);
 
@@ -399,7 +400,8 @@ void FOnlineTurnBasedIOS::LoadAllMatches(FLoadTurnBasedMatchesSignature MatchesL
 		for (GKTurnBasedMatch* match in matches)
 		{
 			NSArray* playerIdentifierArray = FOnlineTurnBasedIOS::GetPlayerIdentifierArrayForMatch(match);
-			[GKPlayer loadPlayersForIdentifiers : playerIdentifierArray withCompletionHandler : ^ (NSArray *players, NSError *nameLoadError)
+			GKLocalPlayer* GKLocalUser = [GKLocalPlayer localPlayer];
+			[GKLocalUser loadFriendsWithIdentifiers:(NSArray<NSString *> *)playerIdentifierArray completionHandler:^(NSArray<GKPlayer *> *players, NSError *nameLoadError)
 			{
 				if (!nameLoadError)
 				{
@@ -431,7 +433,8 @@ void FOnlineTurnBasedIOS::LoadMatchWithID(FString MatchID, FLoadTurnBasedMatchWi
 		if (!error)
 		{
 			NSArray* playerIdentifierArray = FOnlineTurnBasedIOS::GetPlayerIdentifierArrayForMatch(match);
-			[GKPlayer loadPlayersForIdentifiers : playerIdentifierArray withCompletionHandler : ^ (NSArray *players, NSError *nameLoadError)
+			GKLocalPlayer* GKLocalUser = [GKLocalPlayer localPlayer];
+			[GKLocalUser loadFriendsWithIdentifiers:(NSArray<NSString *> *)playerIdentifierArray completionHandler:^(NSArray<GKPlayer *> *players, NSError *nameLoadError)
 			{
 				if (!nameLoadError)
 				{
@@ -542,23 +545,25 @@ void FOnlineTurnBasedIOS::OnMatchEnded(FString MatchID)
 void FOnlineTurnBasedIOS::OnMatchReceivedTurnEvent(FString MatchID, bool BecameActive, void* Match)
 {
 	GKTurnBasedMatch* IOSTurnBasedMatch = (GKTurnBasedMatch*)Match;
-    NSArray* playerIdentifierArray = FOnlineTurnBasedIOS::GetPlayerIdentifierArrayForMatch(IOSTurnBasedMatch);
-    [GKPlayer loadPlayersForIdentifiers : playerIdentifierArray withCompletionHandler : ^ (NSArray *players, NSError *nameLoadError)
-    {
-        if (!nameLoadError)
-        {
-            FTurnBasedMatchPtr PreviousMatchPtr = GetMatchWithID(MatchID);
-            if (PreviousMatchPtr.IsValid())
-            {
-                MatchArray.Remove(PreviousMatchPtr.ToSharedRef());
-            }
-
-            FTurnBasedMatchRef NewMatch = MakeShareable(new FTurnBasedMatchIOS(IOSTurnBasedMatch, players));
-            MatchArray.Add(NewMatch);
-            dispatch_async(dispatch_get_main_queue(), ^
-            {
-                [FIOSAsyncTask CreateTaskWithBlock : ^ bool(void)
-                {
+	
+	NSArray* PlayerIdentifierArray = FOnlineTurnBasedIOS::GetPlayerIdentifierArrayForMatch(IOSTurnBasedMatch);
+	GKLocalPlayer* GKLocalUser = [GKLocalPlayer localPlayer];
+	[GKLocalUser loadFriendsWithIdentifiers:(NSArray<NSString *> *)PlayerIdentifierArray completionHandler:^(NSArray<GKPlayer *> *Players, NSError *NameLoadError)
+	{
+		if (!NameLoadError)
+		{
+			FTurnBasedMatchPtr PreviousMatchPtr = GetMatchWithID(MatchID);
+			if (PreviousMatchPtr.IsValid())
+			{
+				MatchArray.Remove(PreviousMatchPtr.ToSharedRef());
+			}
+			
+			FTurnBasedMatchRef NewMatch = MakeShareable(new FTurnBasedMatchIOS(IOSTurnBasedMatch, Players));
+			MatchArray.Add(NewMatch);
+			dispatch_async(dispatch_get_main_queue(), ^
+						   {
+				[FIOSAsyncTask CreateTaskWithBlock : ^ bool(void)
+				 {
 					if (TurnBasedMatchInterfaceObject)
 					{
 						TArray<uint8> MatchData;
@@ -574,12 +579,12 @@ void FOnlineTurnBasedIOS::OnMatchReceivedTurnEvent(FString MatchID, bool BecameA
 					{
 						EventDelegate.Pin()->OnMatchReceivedTurnEvent(MatchID, BecameActive, Match);
 					}
-                    return true;
-                }];
-
-            });
-        }
-    }];
+					return true;
+				}];
+				
+			});
+		}
+	}];
 }
 
 NSArray* FOnlineTurnBasedIOS::GetPlayerIdentifierArrayForMatch(GKTurnBasedMatch* match)
diff --git a/Engine/Plugins/Online/IOS/OnlineSubsystemIOS/Source/Private/OnlineUserCloudInterfaceIOS.cpp b/Engine/Plugins/Online/IOS/OnlineSubsystemIOS/Source/Private/OnlineUserCloudInterfaceIOS.cpp
--- a/Engine/Plugins/Online/IOS/OnlineSubsystemIOS/Source/Private/OnlineUserCloudInterfaceIOS.cpp
+++ b/Engine/Plugins/Online/IOS/OnlineSubsystemIOS/Source/Private/OnlineUserCloudInterfaceIOS.cpp
@@ -84,7 +84,12 @@
 			CKModifyRecordsOperation *modifyRecords = [[CKModifyRecordsOperation alloc] initWithRecordsToSave:@[record] recordIDsToDelete:nil];
 			modifyRecords.savePolicy = CKRecordSaveAllKeys;
 			modifyRecords.qualityOfService = NSQualityOfServiceUserInitiated;
-			modifyRecords.perRecordCompletionBlock = handler;
+			
+//			modifyRecords.perRecordCompletionBlock = handler; /* Deprecated iOS15. */
+			// Since only 1 record is being saved per run, use the global completion block instead */
+			modifyRecords.modifyRecordsCompletionBlock = ^(NSArray<CKRecord *> *savedRecords, NSArray<CKRecordID *> *deletedRecordIDs, NSError *operationError) {
+				handler([savedRecords firstObject], operationError);
+			};
 			[DB addOperation : modifyRecords];
 
 			return true;
@@ -118,7 +123,7 @@
 			CKQuery* query = [[[CKQuery alloc] initWithRecordType:@"file" predicate:[NSPredicate predicateWithFormat : @"TRUEPREDICATE"]] autorelease];
 			CKQueryOperation* queryOp = [[CKQueryOperation alloc] initWithQuery:query];
 			queryOp.desiredKeys = @[@"record.recordID.recordName"];
-			queryOp.recordFetchedBlock = fetch;
+//			queryOp.recordFetchedBlock = fetch;  /* Deprecated in iOS15.  Now performs once all records are loaded */
 			queryOp.queryCompletionBlock = complete;
 			queryOp.resultsLimit = CKQueryOperationMaximumResults;
 			[DB addOperation : queryOp];
diff --git a/Engine/Plugins/Online/IOS/OnlineSubsystemIOS/Source/Private/TurnBasedMatchmakerDelegateIOS.cpp b/Engine/Plugins/Online/IOS/OnlineSubsystemIOS/Source/Private/TurnBasedMatchmakerDelegateIOS.cpp
--- a/Engine/Plugins/Online/IOS/OnlineSubsystemIOS/Source/Private/TurnBasedMatchmakerDelegateIOS.cpp
+++ b/Engine/Plugins/Online/IOS/OnlineSubsystemIOS/Source/Private/TurnBasedMatchmakerDelegateIOS.cpp
@@ -4,6 +4,7 @@
 #include "OnlineTurnBasedInterfaceIOS.h"
 #include "Interfaces/OnlineTurnBasedInterface.h"
 #include "OnlineSubsystemIOS.h"
+#include <GameKit/GKLocalPlayer.h>
 
 @interface FTurnBasedMatchmakerDelegateIOS()
 {
@@ -41,7 +42,7 @@
 
 			NSMutableArray* playerIdentifierArray = [NSMutableArray array];
 			for (GKTurnBasedParticipant* participant in match.participants) {
-                NSString* PlayerIDString = nil;
+				NSString* PlayerIDString = nil;
 				if ([GKTurnBasedParticipant respondsToSelector:@selector(player)] == YES)
 				{
 					PlayerIDString = FOnlineSubsystemIOS::GetPlayerId(participant.player);
@@ -52,8 +53,12 @@
 				}
 				[playerIdentifierArray addObject : PlayerIDString];
 			}
-			[GKPlayer loadPlayersForIdentifiers : playerIdentifierArray withCompletionHandler : ^ (NSArray *players, NSError *nameLoadError) {
-				if (nameLoadError) {
+			
+			GKLocalPlayer* GKLocalUser = [GKLocalPlayer localPlayer];
+			[GKLocalUser loadFriendsWithIdentifiers:(NSArray<NSString *> *)playerIdentifierArray completionHandler:^(NSArray<GKPlayer *> *players, NSError *nameLoadError)
+			{
+				if (nameLoadError)
+				{
 					[self turnBasedMatchmakerViewController : viewController didFailWithError : nameLoadError];
 					return;
 				}
diff --git a/Engine/Source/Runtime/ApplicationCore/Private/IOS/IOSWindow.cpp b/Engine/Source/Runtime/ApplicationCore/Private/IOS/IOSWindow.cpp
--- a/Engine/Source/Runtime/ApplicationCore/Private/IOS/IOSWindow.cpp
+++ b/Engine/Source/Runtime/ApplicationCore/Private/IOS/IOSWindow.cpp
@@ -23,8 +23,20 @@ void FIOSWindow::Initialize( class FIOSApplication* const Application, const TSh
 {
 	OwningApplication = Application;
 	Definition = InDefinition;
-        
-        Window = (UIWindow *)[[[[UIApplication sharedApplication] connectedScenes] allObjects] firstObject];
+	
+	NSArray<UIScene*> *windowScenes = [[[UIApplication sharedApplication] connectedScenes] allObjects];
+	
+	// TODO: iOS only has one window scene (as of iOS15), however, iPadOS15(?) might have more than one with an external monitor (maybe?).
+	//       Need to investigate and check if looping through windowScene will be required.
+	UIWindowScene *windowScene = (UIWindowScene*)[windowScenes firstObject];
+	for (UIWindow *currentWindow in windowScene.windows)
+	{
+		if (currentWindow.isKeyWindow)
+		{
+			Window = currentWindow;
+			break;
+		}
+	}
 
 #if !PLATFORM_TVOS
 	if(InParent.Get() != NULL)
