Simplify iOS in app purchases code and fix some small issues
- Replaced delegates to communicate from StoreKit to C++ by an internal Objective-C proxy type that sequences requests/notifications on the proper thread
  - StoreKit calls are made on main thread. There is no documentation about thread safety on SKPaymentQueue API so I did not want to take risks. All notifications to the observer are already made on the main thread
  - C++ plugin code is expected to be run on game thread
  - SKRequests delegate methods are received in a worker thread so they need to be sequenced to the game thread
- Reduced the amount of notification functions between C++ and Objective-C layers
- Remove receipts for finalized purchases
- Removed per user cached transactions because iOS/StoreKit is single user. It was confusing and was already lead to some bugs
- Product quantity was ignored (it is possible to buy several items of the same type in one purchase)
- Hidden some private types
- Added C++ RAII wrapper for Objective-C instances to automatically handle manual retain/release

Tested using StoreKit testing in Xcode: https://developer.apple.com/documentation/xcode/setting-up-storekit-testing-in-xcode
Was able to test consumable, non consumable, interrupted and deferred transactions

diff --git a/Engine/Plugins/Online/IOS/OnlineSubsystemIOS/Source/Private/OnlinePurchaseIOS.cpp b/Engine/Plugins/Online/IOS/OnlineSubsystemIOS/Source/Private/OnlinePurchaseIOS.cpp
--- a/Engine/Plugins/Online/IOS/OnlineSubsystemIOS/Source/Private/OnlinePurchaseIOS.cpp
+++ b/Engine/Plugins/Online/IOS/OnlineSubsystemIOS/Source/Private/OnlinePurchaseIOS.cpp
@@ -4,72 +4,376 @@
 #include "OnlineError.h"
 #include "OnlineSubsystemIOS.h"
 #include "Stats/Stats.h"
-#import "OnlineStoreKitHelper.h"
-
-/** Take successful transactions and route them through deferred pipeline */
-#define TEST_DEFERRED_TRANSACTIONS 0
-#define TEST_DEFERRED_DELAY 25.0f
-#if TEST_DEFERRED_TRANSACTIONS
-#include "Containers/Ticker.h"
-#endif // TEST_DEFERRED_TRANSACTIONS
-
-namespace OSSConsoleVariables
-{
-	// CVars
-	TAutoConsoleVariable<int32> CVarSimluateAskToBuy(
-															 TEXT("OSS.AskToBuy"),
-															 0,
-															 TEXT("Simulate ask to buy in iOS\n")
-															 TEXT("1 enable, 0 disable"),
-															 ECVF_Default);
-}
+#include "IOS/ProxyPaymentTransactionObserver.h"
+
+#import <StoreKit/SKError.h>
+#import <StoreKit/SKReceiptRefreshRequest.h>
+#import <StoreKit/SKPaymentTransaction.h>
+#import <StoreKit/SKPayment.h>
+#import <StoreKit/SKPaymentQueue.h>
 
 #define LOCTEXT_NAMESPACE "OnlineSubsystemIOS"
-#define IOSUSER TEXT("IOSUser")
 
-FOnlinePurchaseIOS::FOnlinePurchaseIOS(FOnlineSubsystemIOS* InSubsystem)
-	: StoreHelper(nullptr)
-	, bRestoringTransactions(false)
-	, Subsystem(InSubsystem)
+namespace OnlinePurchaseIOSPrivate
 {
-	UE_LOG_ONLINE_PURCHASE(Log, TEXT( "FOnlinePurchaseIOS::FOnlinePurchaseIOS" ));
+    TAutoConsoleVariable<int32> CVarSimluateAskToBuy(TEXT("OSS.AskToBuy"),
+                                                     0,
+                                                     TEXT("Simulate ask to buy in iOS\n")
+                                                     TEXT("1 enable, 0 disable"),
+                                                     ECVF_Default);
+
+    /**
+     * Convert an Apple SKPaymentTransaction receipt into a string
+     *
+     * @return hex encoded string with opaque data representing a completed transaction
+     */
+    FString ConvertReceiptToString()
+    {
+        FString ReceiptData;
+        
+        NSURL* nsReceiptUrl = [[NSBundle mainBundle] appStoreReceiptURL];
+        NSData* nsReceiptData = [NSData dataWithContentsOfURL : nsReceiptUrl];
+        if (nsReceiptData)
+        {
+            NSString* nsEncodedReceiptData = [nsReceiptData base64EncodedStringWithOptions : NSDataBase64EncodingEndLineWithLineFeed];
+            ReceiptData = nsEncodedReceiptData;
+            UE_LOG_ONLINE_STOREV2(VeryVerbose, TEXT("ConvertReceiptToString %s"), *ReceiptData);
+        }
+        else
+        {
+            UE_LOG_ONLINE_STOREV2(Log, TEXT("ConvertReceiptToString: No receipt data found for transaction"));
+        }
+        return ReceiptData;
+    }
+}
+
+/**
+ * Intermediate type to communicate SKPaymentTransaction information from Objective-C to
+ * to C++ implementation
+ */
+struct FOnlinePurchaseTransactionIOS
+{
+    explicit FOnlinePurchaseTransactionIOS(SKPaymentTransaction* Transaction);
+    
+    /** @return a string that prints useful debug information about this transaction */
+    FString ToDebugString() const
+    {
+        return FString::Printf(TEXT("OfferId: %s TransactionId: %s ReceiptData: %s Error:%s"),
+                        *OfferId,
+                        *TransactionIdentifier,
+                        *ReceiptData,
+                        *ErrorStr);
+    }
+    
+    /** @return offer id for this transaction */
+    const FString& GetOfferId() const { return OfferId; }
+    /** @return receipt data for this transaction */
+    const FString& GetReceiptData() const { return ReceiptData; }
+    /** @return error string for this transaction, if applicable */
+    const FString& GetErrorStr() const { return ErrorStr; }
+    /** @return quantity of items requested of the same type */
+    int GetQuantity() const { return Quantity; }
+    /** @return transaction identifier for this purchase */
+    const FString& GetTransactionIdentifier() const { return TransactionIdentifier; }
+    /** @return the SKPaymentTransaction object so we can finish the transaction */
+    SKPaymentTransaction* GetPaymentTransaction() const { return static_cast<SKPaymentTransaction*>(PaymentTransaction); }
+private:
+    
+    /** iTunesConnect offer id */
+    FString OfferId;
+    /** Opaque store receipt data */
+    FString ReceiptData;
+    /** Error on the transaction, if applicable */
+    FString ErrorStr;
+    /** Unique transaction identifier */
+    FString TransactionIdentifier;
+    /** Quantity of items requested*/
+    int Quantity;
+    /** Platform transaction object*/
+    TRetainedObjCInstance<SKPaymentTransaction*> PaymentTransaction;
+};
+
+FOnlinePurchaseTransactionIOS::FOnlinePurchaseTransactionIOS(SKPaymentTransaction* Transaction)
+    : ReceiptData(OnlinePurchaseIOSPrivate::ConvertReceiptToString())
+    , TransactionIdentifier(Transaction.transactionIdentifier)
+    , Quantity(0)
+    , PaymentTransaction(Transaction)
+{
+    if(Transaction.error)
+    {
+        FString ErrorRaw([Transaction.error localizedDescription]);
+        if (!ErrorRaw.IsEmpty())
+        {
+            int32 ErrorCode([Transaction.error code]);
+            FString ErrorDomain([Transaction.error domain]);
+            ErrorStr = FString::Printf(TEXT("%s [%s:%d]"), *ErrorRaw, *ErrorDomain, ErrorCode);
+        }
+    }
+    
+    OfferId = Transaction.payment.productIdentifier;
+    Quantity = Transaction.payment.quantity;
+}
+
+/**
+ * Info used to cache and track order in progress.
+ */
+struct FOnlinePurchaseInProgressTransactionIOS
+{
+    FOnlinePurchaseInProgressTransactionIOS(const FPurchaseCheckoutRequest& InCheckoutRequest, const FOnPurchaseCheckoutComplete InCheckoutCompleteDelegate)
+        : CheckoutRequest(InCheckoutRequest)
+        , CheckoutCompleteDelegate(InCheckoutCompleteDelegate)
+    {
+    }
+    
+    bool HasRequestForProduct(const FString& ProductId) const
+    {
+        const TArray<FPurchaseCheckoutRequest::FPurchaseOfferEntry>& Offers = CheckoutRequest.PurchaseOffers;
+        return (Offers.Num() > 0 ? Offers[0].OfferId == ProductId : false);
+    }
+    
+    /** Checkout info for the order */
+    FPurchaseCheckoutRequest CheckoutRequest;
+    /** Delegate to call on completion */
+    FOnPurchaseCheckoutComplete CheckoutCompleteDelegate;
+};
+
+/**
+ * Proxy class that notifies updates from SKPaymentQueue to FOnlinePurchaseIOS on the game 
+ * thread and invokes SKPaymentQueue methods on main thread
+ */
+@interface FStoreKitPurchaseProxy : NSObject<SKPaymentTransactionObserver>
+{
+    FOnlinePurchaseIOS* _PurchaseReceiver;
+};
+@end
+
+@implementation FStoreKitPurchaseProxy
+
+////////////////////////////////////////////////////////////////////
+/// SKPaymentTransactionObserver implementation
+-(void)paymentQueue: (SKPaymentQueue*)Queue updatedTransactions: (NSArray*)Transactions
+{
+    UE_LOG_ONLINE_STOREV2(Log, TEXT("FStoreKitPurchaseProxy::updatedTransactions"));
+    for(SKPaymentTransaction* Transaction in Transactions)
+    {
+        [self UpdatedTransaction : Transaction];
+    }
+}
+
+-(void)paymentQueue: (SKPaymentQueue*)Queue removedTransactions: (NSArray*)Transactions
+{
+    UE_LOG_ONLINE_STOREV2(Log, TEXT("FStoreKitPurchaseProxy::removedTransaction"));
+}
+
+-(void)paymentQueue: (SKPaymentQueue*)Queue restoreCompletedTransactionsFailedWithError: (NSError*)Error
+{
+    UE_LOG_ONLINE_STOREV2(Log, TEXT("FStoreKitPurchaseProxy::restoreCompletedTransactionsFailedWithError"));
+
+    EPurchaseTransactionState CompletionState = [FStoreKitPurchaseProxy TranslateError:Error];
+    
+    // Notifications from SKPaymentQueue are received on main thread. Use a task to notify from game thread
+    [FIOSAsyncTask CreateTaskWithBlock : ^ bool(void)
+    {
+        if(_PurchaseReceiver)
+        {
+            _PurchaseReceiver->OnRestoreTransactionsComplete(CompletionState);
+        }
+        return true;
+    }];
 }
 
-FOnlinePurchaseIOS::FOnlinePurchaseIOS()
+-(void)paymentQueueRestoreCompletedTransactionsFinished: (SKPaymentQueue*)Queue
+{
+    UE_LOG_ONLINE_STOREV2(Log, TEXT("FStoreKitPurchaseProxy::paymentQueueRestoreCompletedTransactionsFinished"));
+
+    // Notifications from SKPaymentQueue are received on main thread. Use a task to notify from game thread
+    [FIOSAsyncTask CreateTaskWithBlock : ^ bool(void)
+    {
+        if(_PurchaseReceiver)
+        {
+            _PurchaseReceiver->OnRestoreTransactionsComplete(EPurchaseTransactionState::Restored);
+        }
+        return true;
+    }];
+}
+
+////////////////////////////////////////////////////////////////////
+/// FStoreKitPurchaseProxy methods
+
+- (id)initWithReceiver: (FOnlinePurchaseIOS*)PurchaseReceiver
+{
+    _PurchaseReceiver = PurchaseReceiver;
+    [FProxyPaymentTransactionObserver sharedInstance].ossObserver = self;
+
+    // Notifications from SKPaymentQueue are received on main thread. Use dispatch_async on main queue to serialize operations
+    dispatch_async(dispatch_get_main_queue(), ^
+    {
+        SKPaymentQueue* queue = [SKPaymentQueue defaultQueue];
+        [self paymentQueue: queue updatedTransactions: [queue transactions]];
+    });
+    return self;
+}
+
+-(void)Shutdown
+{
+    _PurchaseReceiver = nullptr;
+    [FProxyPaymentTransactionObserver sharedInstance].ossObserver = nil;
+}
+
+- (bool)CanMakePayments
+{
+    return [SKPaymentQueue canMakePayments] != FALSE;
+}
+
++(EPurchaseTransactionState)TranslateError: (NSError*)Error
+{
+    EPurchaseTransactionState TranslatedError = EPurchaseTransactionState::Failed;
+    switch (Error.code)
+    {
+        case SKErrorPaymentCancelled:
+            TranslatedError = EPurchaseTransactionState::Canceled;
+            break;
+        case SKErrorClientInvalid:
+        case SKErrorStoreProductNotAvailable:
+        case SKErrorPaymentInvalid:
+            TranslatedError = EPurchaseTransactionState::Invalid;
+            break;
+        case SKErrorPaymentNotAllowed:
+            TranslatedError = EPurchaseTransactionState::NotAllowed;
+            break;
+    }
+    return TranslatedError;
+}
+
+-(void)UpdatedTransaction : (SKPaymentTransaction*)Transaction
+{
+    TOptional<EPurchaseTransactionState> TranslatedState;
+    SKPaymentTransaction* ActualTransaction = Transaction;
+    SKPaymentTransaction* FinishTransaction = nil;
+    
+    switch ([Transaction transactionState])
+    {
+    case SKPaymentTransactionStatePurchased:
+        UE_LOG_ONLINE_STOREV2(Log, TEXT("FStoreKitPurchaseProxy::updatedTransaction SKPaymentTransactionStatePurchased"));
+
+        if (FParse::Param(FCommandLine::Get(), TEXT("disableiosredeem")))
+        {
+            UE_LOG_ONLINE_STOREV2(Log, TEXT("FStoreKitPurchaseProxy::completedTransaction (disabled)"));
+        }
+        else
+        {
+            TranslatedState = EPurchaseTransactionState::Purchased;
+        }
+        break;
+    case SKPaymentTransactionStateFailed:
+        UE_LOG_ONLINE_STOREV2(Log, TEXT("FStoreKitPurchaseProxy::updatedTransaction SKPaymentTransactionStateFailed"));
+        TranslatedState = [FStoreKitPurchaseProxy TranslateError:ActualTransaction.error];
+        FinishTransaction = ActualTransaction;
+        break;
+    case SKPaymentTransactionStateRestored:
+        UE_LOG_ONLINE_STOREV2(Log, TEXT("FStoreKitPurchaseProxy::updatedTransaction SKPaymentTransactionStateRestored"));
+        TranslatedState = EPurchaseTransactionState::Restored;
+        FinishTransaction = ActualTransaction;
+        ActualTransaction = ActualTransaction.originalTransaction;
+        break;
+    case SKPaymentTransactionStatePurchasing:
+        UE_LOG_ONLINE_STOREV2(Log, TEXT("FStoreKitPurchaseProxy::updatedTransaction SKPaymentTransactionStatePurchasing"));
+        break;
+    case SKPaymentTransactionStateDeferred:
+        UE_LOG_ONLINE_STOREV2(Log, TEXT("FStoreKitPurchaseProxy::updatedTransaction SKPaymentTransactionStateDeferred"));
+        TranslatedState = EPurchaseTransactionState::Deferred;
+        FinishTransaction = ActualTransaction;
+        break;
+    default:
+        UE_LOG_ONLINE_STOREV2(Warning, TEXT("FStoreKitPurchaseProxy::updatedTransaction unhandled state: %d"), [Transaction transactionState]);
+        break;
+    }
+
+    if(TranslatedState)
+    {
+        FOnlinePurchaseTransactionIOS TransactionData(ActualTransaction);
+        // Notifications from SKPaymentQueue are received on main thread. Use a task to notify from game thread
+        [FIOSAsyncTask CreateTaskWithBlock : ^ bool(void)
+        {
+            if(_PurchaseReceiver)
+            {
+                _PurchaseReceiver->OnTransactionComplete(*TranslatedState, TransactionData);
+            }
+            return true;
+        }];
+    }
+    
+    if(FinishTransaction)
+    {
+        UE_LOG_ONLINE_STOREV2(Log, TEXT("FStoreKitPurchaseProxy::updatedTransaction transaction finished"));
+        [[SKPaymentQueue defaultQueue] finishTransaction:FinishTransaction];
+    }
+}
+
+-(void)RestorePurchases
+{
+    // Notifications from SKPaymentQueue are received on main thread. Use dispatch_async on main queue to serialize operations
+    dispatch_async(dispatch_get_main_queue(), ^
+    {
+        [[SKPaymentQueue defaultQueue] restoreCompletedTransactions];
+    });
+}
+
+-(void)MakePurchase: (SKProduct*)Product withQuantity: (int)Quantity simulateAskToBuy: (bool) bAskToBuy
+{
+    UE_LOG_ONLINE_STOREV2(Log, TEXT("FStoreKitPurchaseProxy::makePurchase by SKProduct"));
+
+    // Notifications from SKPaymentQueue are received on main thread. Use dispatch_async on main queue to serialize operations
+    dispatch_async(dispatch_get_main_queue(), ^
+    {
+        SKMutablePayment* Payment = [SKMutablePayment paymentWithProduct:Product];
+        Payment.quantity = Quantity;
+        // simulatesAskToBuyInSandbox is only effective while Sandbox testing
+        Payment.simulatesAskToBuyInSandbox = bAskToBuy;
+
+        [[SKPaymentQueue defaultQueue] addPayment:Payment];
+    });
+}
+
+-(void)FinalizeTransaction: (SKPaymentTransaction*)Transaction
+{
+    UE_LOG_ONLINE_STOREV2(Log, TEXT("FStoreKitPurchaseProxy::finalizeTransaction"));
+
+    // Notifications from SKPaymentQueue are received on main thread. Use dispatch_async on main queue to serialize operations
+    dispatch_async(dispatch_get_main_queue(), ^
+    {
+        [[SKPaymentQueue defaultQueue] finishTransaction:Transaction];
+    });
+}
+
+@end
+
+FOnlinePurchaseIOS::FOnlinePurchaseIOS(FOnlineSubsystemIOS* InSubsystem)
+	: bRestoringTransactions(false)
+	, Subsystem(InSubsystem)
 {
 	UE_LOG_ONLINE_PURCHASE(Log, TEXT( "FOnlinePurchaseIOS::FOnlinePurchaseIOS" ));
+    StoreKitProxy = [[[FStoreKitPurchaseProxy alloc] initWithReceiver:this] autorelease];
 }
 
+////////////////////////////////////////////////////////////////////
+/// FOnlinePurchaseIOS methods
+
 FOnlinePurchaseIOS::~FOnlinePurchaseIOS()
 {
+    [StoreKitProxy Shutdown];
 }
 
-void FOnlinePurchaseIOS::InitStoreKit(FStoreKitHelperV2* InStoreKit)
+void FOnlinePurchaseIOS::DumpAppReceipt()
 {
-	StoreHelper = InStoreKit;
-
-	FOnProductsRequestResponseDelegate OnProductsRequestResponseDelegate = FOnProductsRequestResponseDelegate::CreateThreadSafeSP(this, &FOnlinePurchaseIOS::OnProductPurchaseRequestResponse);
-	[StoreHelper AddOnProductRequestResponse: OnProductsRequestResponseDelegate];
-	
-	FOnTransactionCompleteIOSDelegate OnTransactionCompleteResponseDelegate = FOnTransactionCompleteIOSDelegate::CreateThreadSafeSP(this, &FOnlinePurchaseIOS::OnTransactionCompleteResponse);
-	[StoreHelper AddOnTransactionComplete: OnTransactionCompleteResponseDelegate];
-	
-	FOnTransactionRestoredIOSDelegate OnTransactionRestoredDelegate = FOnTransactionRestoredIOSDelegate::CreateThreadSafeSP(this, &FOnlinePurchaseIOS::OnTransactionRestored);
-	[StoreHelper AddOnTransactionRestored: OnTransactionRestoredDelegate];
-	
-	FOnRestoreTransactionsCompleteIOSDelegate OnRestoreTransactionsCompleteDelegate = FOnRestoreTransactionsCompleteIOSDelegate::CreateThreadSafeSP(this, &FOnlinePurchaseIOS::OnRestoreTransactionsComplete);
-	[StoreHelper AddOnRestoreTransactionsComplete: OnRestoreTransactionsCompleteDelegate];
-	
-	FOnTransactionProgressDelegate OnTransactionPurchasingDelegate = FOnTransactionProgressDelegate::CreateThreadSafeSP(this, &FOnlinePurchaseIOS::OnTransactionInProgress);
-	[StoreHelper AddOnPurchaseInProgress: OnTransactionPurchasingDelegate];
-	
-	FOnTransactionProgressDelegate OnTransactionDeferredDelegate = FOnTransactionProgressDelegate::CreateThreadSafeSP(this, &FOnlinePurchaseIOS::OnTransactionDeferred);
-	[StoreHelper AddOnTransactionDeferred: OnTransactionDeferredDelegate];
+    FString receiptData = OnlinePurchaseIOSPrivate::ConvertReceiptToString();
+    UE_LOG_ONLINE_STOREV2(Verbose, TEXT("FOnlinePurchaseIOS::DumpAppReceipt"));
+    UE_LOG_ONLINE_STOREV2(Verbose, TEXT("%s"), *receiptData);
 }
 
-bool FOnlinePurchaseIOS::IsAllowedToPurchase(const FUniqueNetId& UserId)
+bool FOnlinePurchaseIOS::IsAllowedToPurchase(const FUniqueNetId& /*UserId*/)
 {
-	bool bCanMakePurchases = [SKPaymentQueue canMakePayments];
+	bool bCanMakePurchases = [StoreKitProxy CanMakePayments];
 	UE_LOG_ONLINE_PURCHASE(Verbose, TEXT("FOnlinePurchaseIOS::IsAllowedToPurchase %s"), *LexToString(bCanMakePurchases));
 	return bCanMakePurchases;
 }
@@ -79,73 +383,63 @@ void FOnlinePurchaseIOS::Checkout(const FUniqueNetId& UserId, const FPurchaseChe
 	bool bStarted = false;
 	FText ErrorMessage;
 
-	TSharedRef<FOnlinePurchasePendingTransactionIOS> RequestedTransaction = MakeShareable(new FOnlinePurchasePendingTransactionIOS(CheckoutRequest, UserId, EPurchaseTransactionState::NotStarted, Delegate));
-
-	if (IsAllowedToPurchase(UserId))
+    if (IsAllowedToPurchase(UserId))
 	{
-		const FString UserIdStr = IOSUSER;
-		const TSharedRef<FOnlinePurchasePendingTransactionIOS>* UserPendingTransaction = PendingTransactions.Find(UserIdStr);
-		if (UserPendingTransaction == nullptr)
+		if (InProgressTransaction == nullptr)
 		{
 			FOnlineStoreIOSPtr StoreInterface = StaticCastSharedPtr<FOnlineStoreIOS>(Subsystem->GetStoreV2Interface());
 			if (StoreInterface.IsValid())
 			{
-				TSharedRef<FOnlinePurchasePendingTransactionIOS> PendingTransaction = PendingTransactions.Add(UserIdStr, RequestedTransaction);
-				PendingTransaction->StartProcessing();
-				
-				// autoreleased NSMutableArray to hold products
-				int32 NumOffers = CheckoutRequest.PurchaseOffers.Num();
-
-				NSMutableArray<SKProduct*>* ProductSet = [NSMutableArray arrayWithCapacity: NumOffers];
-				for (int32 OfferIdx = 0; OfferIdx < NumOffers; OfferIdx++)
-				{
-					const FPurchaseCheckoutRequest::FPurchaseOfferEntry& Offer = CheckoutRequest.PurchaseOffers[OfferIdx];
-
-					SKProduct* Product = StoreInterface->GetSKProductByOfferId(Offer.OfferId);
-					if (Product)
-					{
-						[ProductSet addObject:Product];
-					}
-				}
-				
-				bool bAskToBuy = false;
+                int32 NumOffers = CheckoutRequest.PurchaseOffers.Num();
+                
+                SKProduct* Product = nil;
+                if(NumOffers > 0)
+                {
+                    const FPurchaseCheckoutRequest::FPurchaseOfferEntry& Offer = CheckoutRequest.PurchaseOffers[0];
+                    if (NumOffers > 1)
+                    {
+                        UE_LOG_ONLINE_PURCHASE(Warning, TEXT("StoreKit does not support multiple different products in one transaction. Only %s will be requested"), *Offer.OfferId);
+                    }
+                    Product = StoreInterface->GetSKProductByOfferId(Offer.OfferId);
+                
+                    if (Product)
+                    {
+                        bool bAskToBuy = false;
 #if !UE_BUILD_SHIPPING
-				bAskToBuy = OSSConsoleVariables::CVarSimluateAskToBuy.GetValueOnGameThread() == 1;
+                        bAskToBuy = OnlinePurchaseIOSPrivate::CVarSimluateAskToBuy.GetValueOnGameThread() == 1;
 #endif
-				
-				if ([ProductSet count] > 0)
-				{
-					dispatch_async(dispatch_get_main_queue(), ^
-					{
-					    // Purchase the product through the StoreKit framework
-						[StoreHelper makePurchase:ProductSet SimulateAskToBuy:bAskToBuy];
-					});
-					bStarted = true;
-				}
-				else
-				{
-					ErrorMessage = NSLOCTEXT("IOSPurchase", "ErrorNoOffersSpecified", "Failed to checkout, no offers given.");
-					RequestedTransaction->PendingPurchaseInfo.TransactionState = EPurchaseTransactionState::Failed;
-				}
+                        
+                        [StoreKitProxy MakePurchase:Product withQuantity:Offer.Quantity simulateAskToBuy:bAskToBuy];
+                        InProgressTransaction = MakeShared<const FOnlinePurchaseInProgressTransactionIOS>(CheckoutRequest, Delegate);
+                        bStarted = true;
+                    }
+                }
+                if(!bStarted)
+                {
+                    ErrorMessage = NSLOCTEXT("IOSPurchase", "ErrorNoOffersSpecified", "Failed to checkout, no valid offers given.");
+                }
 			}
-		}
+            else
+            {
+                ErrorMessage = NSLOCTEXT("IOSPurchase", "ErrorPurchaseNotAllowed", "Failed to checkout, invalid FOnlineStoreIOS instance.");
+            }		
+        }
 		else
 		{
 			ErrorMessage = NSLOCTEXT("IOSPurchase", "ErrorTransactionInProgress", "Failed to checkout, user has in progress transaction.");
-			RequestedTransaction->PendingPurchaseInfo.TransactionState = EPurchaseTransactionState::Failed;
 		}
 	}
 	else
 	{
 		ErrorMessage = NSLOCTEXT("IOSPurchase", "ErrorPurchaseNotAllowed", "Failed to checkout, user not allowed to purchase.");
-		RequestedTransaction->PendingPurchaseInfo.TransactionState = EPurchaseTransactionState::Failed;
 	}
 
 	if (!bStarted)
 	{
-		TSharedRef<FPurchaseReceipt> FailReceipt = RequestedTransaction->GenerateReceipt();
+		TSharedRef<FPurchaseReceipt> FailReceipt = GenerateFailReceipt(CheckoutRequest);
 		
-		Subsystem->ExecuteNextTick([ErrorMessage, FailReceipt, Delegate]()
+        // Notify failure on next tick
+		Subsystem->ExecuteNextTick([ErrorMessage, FailReceipt = FailReceipt, Delegate]()
 		{
 			FOnlineError Error(ErrorMessage);
 			Delegate.ExecuteIfBound(Error, FailReceipt);
@@ -163,18 +457,24 @@ void FOnlinePurchaseIOS::FinalizePurchase(const FUniqueNetId& UserId, const FStr
 {
 	UE_LOG_ONLINE_PURCHASE(Log, TEXT("FOnlinePurchaseIOS::FinalizePurchase %s %s"), *UserId.ToString(), *ReceiptId);
 
-	const FString ReceiptIdCopy(ReceiptId);
-	dispatch_async(dispatch_get_main_queue(), ^
-	{
-		// Purchase the product through the StoreKit framework
-		[StoreHelper finalizeTransaction:ReceiptIdCopy];
-	});
+    int IndexToRemove = Transactions.IndexOfByPredicate([&ReceiptId](const FKnownTransaction& Entry) { return Entry.Receipt->TransactionId == ReceiptId; });
+    if(IndexToRemove != INDEX_NONE)
+    {
+        const FKnownTransaction& Transaction = Transactions[IndexToRemove];
+        
+        // Restored transactions are already marked finalized when received but kept as known transactions so user code can process them
+        if(Transaction.Receipt->TransactionState != EPurchaseTransactionState::Restored)
+        {
+            [StoreKitProxy FinalizeTransaction:Transaction.PaymentTransaction];
+        }
+        Transactions.RemoveAt(IndexToRemove);
+    }
 }
 
 void FOnlinePurchaseIOS::RedeemCode(const FUniqueNetId& UserId, const FRedeemCodeRequest& RedeemCodeRequest, const FOnPurchaseRedeemCodeComplete& Delegate)
 {
 	FOnlineError Result;
-	Delegate.ExecuteIfBound(Result, MakeShareable(new FPurchaseReceipt()));
+	Delegate.ExecuteIfBound(Result, MakeShared<FPurchaseReceipt>());
 }
 
 void FOnlinePurchaseIOS::QueryReceipts(const FUniqueNetId& UserId, bool bRestoreReceipts, const FOnQueryReceiptsComplete& Delegate)
@@ -185,21 +485,10 @@ void FOnlinePurchaseIOS::QueryReceipts(const FUniqueNetId& UserId, bool bRestore
 	{
 		if (!bRestoringTransactions)
 		{
-			// Restore purchases, adding them to the offline receipts array for later redemption
-			if (StoreHelper)
-			{
-				QueryReceiptsComplete = Delegate;
-				bTriggerDelegate = false;
-				bRestoringTransactions = true;
-				dispatch_async(dispatch_get_main_queue(), ^
-				{
-					[StoreHelper restorePurchases];
-				});
-			}
-			else
-			{
-				bSuccess = false;
-			}
+            QueryReceiptsComplete = Delegate;
+            bTriggerDelegate = false;
+            bRestoringTransactions = true;
+            [StoreKitProxy RestorePurchases];
 		}
 		else
 		{
@@ -222,101 +511,72 @@ void FOnlinePurchaseIOS::GetReceipts(const FUniqueNetId& UserId, TArray<FPurchas
 {
 	OutReceipts.Empty();
 
-	// Add the cached list of user purchases
-    const FString UserIdStr = IOSUSER;
-	const TArray<TSharedRef<FPurchaseReceipt>>* UserCompletedTransactions = CompletedTransactions.Find(UserIdStr);
-	if (UserCompletedTransactions != nullptr)
-	{
-		for (int32 Idx = 0; Idx < UserCompletedTransactions->Num(); Idx++)
-		{
-			OutReceipts.Add(*(*UserCompletedTransactions)[Idx]);
-		}
-	}
-	
-	// Add purchases completed while "offline"
-	for (int32 Idx = 0; Idx < OfflineTransactions.Num(); Idx++)
+	for (const FKnownTransaction& Transaction: Transactions)
 	{
-		OutReceipts.Add(*OfflineTransactions[Idx]);
+		OutReceipts.Add(*Transaction.Receipt);
 	}
 }
 
-void FOnlinePurchaseIOS::OnProductPurchaseRequestResponse(SKProductsResponse* Response, const FOnQueryOnlineStoreOffersComplete& CompletionDelegate)
-{
-	UE_LOG_ONLINE_PURCHASE(Log, TEXT("FOnlinePurchaseIOS::OnProductPurchaseRequestResponse"));
-}
-
-void FOnlinePurchaseIOS::OnTransactionCompleteResponse(EPurchaseTransactionState Result, const FStoreKitTransactionData& TransactionData)
+void FOnlinePurchaseIOS::OnTransactionComplete(EPurchaseTransactionState Result, const FOnlinePurchaseTransactionIOS& TransactionData)
 {
 	UE_LOG_ONLINE_PURCHASE(Log, TEXT("FOnlinePurchaseIOS::OnTransactionCompleteResponse Result=%s TransactionData=[%s]"), LexToString(Result), *TransactionData.ToDebugString());
 	
-	FString UserIdStr = IOSUSER;
-	const TSharedRef<FOnlinePurchasePendingTransactionIOS>* UserPendingTransactionPtr = PendingTransactions.Find(UserIdStr);
-	if (UserPendingTransactionPtr != nullptr)
+    bool bIsInProgressTransaction = InProgressTransaction != nullptr && InProgressTransaction->HasRequestForProduct(TransactionData.GetOfferId());
+    bool bIsResoredTransaction = (Result == EPurchaseTransactionState::Restored);
+    
+	if (bIsInProgressTransaction && !bIsResoredTransaction)
 	{
-#if TEST_DEFERRED_TRANSACTIONS
-		if (Result == EPurchaseTransactionState::Restored || Result == EPurchaseTransactionState::Purchased)
-		{
-			// Route a successful purchase to the deferred code
-			OnTransactionDeferred(TransactionData);
-			return;
-		}
-#endif // TEST_DEFERRED_TRANSACTIONS
-
-		const TSharedRef<FOnlinePurchasePendingTransactionIOS> UserPendingTransaction = *UserPendingTransactionPtr;
-		
-		UserPendingTransaction->AddCompletedOffer(Result, TransactionData);
-		if (UserPendingTransaction->AreAllOffersComplete())
-		{
-			FOnlineError FinalResult;
-			EPurchaseTransactionState FinalState = UserPendingTransaction->GetFinalTransactionState();
-			
-			UserPendingTransaction->PendingPurchaseInfo.TransactionState = FinalState;
-			// UserPendingTransaction->PendingPurchaseInfo.TransactionId; purposefully blank
-			
-			const FString& ErrorStr = TransactionData.GetErrorStr();
-			switch (FinalState)
-			{
-				case EPurchaseTransactionState::Failed:
-					FinalResult.SetFromErrorCode(TEXT("com.epicgames.purchase.failure"));
-					FinalResult.ErrorMessage = !ErrorStr.IsEmpty() ? FText::FromString(ErrorStr) : LOCTEXT("IOSTransactionFailed", "TransactionFailed");
-					break;
-				case EPurchaseTransactionState::Canceled:
-					FinalResult.SetFromErrorCode(TEXT("com.epicgames.catalog_helper.user_cancelled"));
-					FinalResult.ErrorMessage = !ErrorStr.IsEmpty() ? FText::FromString(ErrorStr) : LOCTEXT("IOSTransactionCancel", "TransactionCanceled");
-					break;
-				case EPurchaseTransactionState::Purchased:
-					FinalResult.bSucceeded = true;
-					break;
-				default:
-					UE_LOG_ONLINE_PURCHASE(Warning, TEXT("Unexpected state after purchase %d"), (int)FinalState);
-					FinalResult.SetFromErrorCode(TEXT("com.epicgames.purchase.unexpected_state"));
-					FinalResult.ErrorMessage = !ErrorStr.IsEmpty() ? FText::FromString(ErrorStr) : LOCTEXT("UnexpectedState", "Unexpected purchase result");
-					UserPendingTransaction->PendingPurchaseInfo.TransactionState = EPurchaseTransactionState::Failed;
-					break;
-			}
-			
-			TSharedRef<FPurchaseReceipt> FinalReceipt = UserPendingTransaction->GenerateReceipt();
-			
-			TArray< TSharedRef<FPurchaseReceipt> >& UserCompletedTransactions = CompletedTransactions.FindOrAdd(UserIdStr);
-			
-			PendingTransactions.Remove(UserIdStr);
-			UserCompletedTransactions.Add(FinalReceipt);
-			
-			UserPendingTransaction->CheckoutCompleteDelegate.ExecuteIfBound(FinalResult, FinalReceipt);
-		}
+        FOnlineError FinalResult;
+        const FString& ErrorStr = TransactionData.GetErrorStr();
+
+        switch (Result)
+        {
+            case EPurchaseTransactionState::Failed:
+                FinalResult.SetFromErrorCode(TEXT("com.epicgames.purchase.failure"));
+                FinalResult.ErrorMessage = !ErrorStr.IsEmpty() ? FText::FromString(ErrorStr) : LOCTEXT("IOSTransactionFailed", "TransactionFailed");
+                break;
+            case EPurchaseTransactionState::Canceled:
+                FinalResult.SetFromErrorCode(TEXT("com.epicgames.catalog_helper.user_cancelled"));
+                FinalResult.ErrorMessage = !ErrorStr.IsEmpty() ? FText::FromString(ErrorStr) : LOCTEXT("IOSTransactionCancel", "TransactionCanceled");
+                break;
+            case EPurchaseTransactionState::Purchased:
+                FinalResult.bSucceeded = true;
+                break;
+            case EPurchaseTransactionState::Deferred:
+                FinalResult.SetFromErrorCode(TEXT("com.epicgames.purchase.deferred"));
+                FinalResult.ErrorMessage = !ErrorStr.IsEmpty() ? FText::FromString(ErrorStr) : LOCTEXT("IOSTransactionDeferred", "Transaction awaiting approval.");
+                break;
+            default:
+                UE_LOG_ONLINE_PURCHASE(Warning, TEXT("Unexpected state after purchase %d"), (int)Result);
+                FinalResult.SetFromErrorCode(TEXT("com.epicgames.purchase.unexpected_state"));
+                FinalResult.ErrorMessage = !ErrorStr.IsEmpty() ? FText::FromString(ErrorStr) : LOCTEXT("UnexpectedState", "Unexpected purchase result");
+                break;
+        }
+        
+        TSharedRef<FPurchaseReceipt> Receipt = GenerateReceipt(Result, InProgressTransaction->CheckoutRequest, TransactionData);
+        
+        if(FinalResult.bSucceeded)
+        {
+            Receipt = TryStoreTransactionAndGetReceipt(Receipt, TransactionData);
+        }
+        InProgressTransaction->CheckoutCompleteDelegate.ExecuteIfBound(FinalResult, Receipt);
+        InProgressTransaction.Reset();
 	}
 	else
 	{
-		// Transactions that come in during login or other non explicit purchase moments are added to a receipts list for later redemption
-		UE_LOG_ONLINE_PURCHASE(Log, TEXT("Pending transaction completed offline"));
+		// Transactions other than the in progress one are stored for later redemption. Those can be received at login or when restoring transactions
+        // We have no interest on other failure states because they won't store a redeemable receipt
 		if (Result == EPurchaseTransactionState::Restored || Result == EPurchaseTransactionState::Purchased)
 		{
-			TSharedRef<FPurchaseReceipt> OfflineReceipt = FOnlinePurchasePendingTransactionIOS::GenerateReceipt(Result, TransactionData);
-			OfflineTransactions.Add(OfflineReceipt);
-
-			// Queue this user to be updated about this next-tick on the game-thread, if they're not mid-purchase
-			TWeakPtr<FOnlinePurchaseIOS, ESPMode::ThreadSafe> WeakThis(AsShared());
-			Subsystem->ExecuteNextTick([WeakThis]()
+            TSharedRef<FPurchaseReceipt> Receipt = GenerateOfflineReceipt(Result, TransactionData);
+            TryStoreTransactionAndGetReceipt(Receipt, TransactionData);
+        }
+        
+        if(Result == EPurchaseTransactionState::Purchased)
+        {
+		    UE_LOG_ONLINE_PURCHASE(Log, TEXT("Deferred transaction finished or unfinished transaction was received"));
+            TWeakPtr<FOnlinePurchaseIOS> WeakThis = AsShared();
+            Subsystem->ExecuteNextTick([WeakThis]()
 			{
 				FOnlinePurchaseIOSPtr StrongThis = WeakThis.Pin();
 				if (StrongThis.IsValid())
@@ -326,34 +586,37 @@ void FOnlinePurchaseIOS::OnTransactionCompleteResponse(EPurchaseTransactionState
 				}
 			});
 		}
+        else if(Result == EPurchaseTransactionState::Restored)
+        {
+            UE_LOG_ONLINE_PURCHASE(Log, TEXT("Restored transaction"));
+        }
 	}
 }
 
-void FOnlinePurchaseIOS::OnTransactionRestored(const FStoreKitTransactionData& TransactionData)
+TSharedRef<FPurchaseReceipt> FOnlinePurchaseIOS::TryStoreTransactionAndGetReceipt(const TSharedRef<FPurchaseReceipt>& Receipt, const FOnlinePurchaseTransactionIOS& Transaction)
 {
-	UE_LOG_ONLINE_PURCHASE(Verbose, TEXT("FOnlinePurchaseIOS::OnTransactionRestored TransactionData=[%s]"), *TransactionData.ToDebugString());
-
-	// Single item restored amongst a group of items
-	TSharedRef<FPurchaseReceipt> OfflineReceipt = FOnlinePurchasePendingTransactionIOS::GenerateReceipt(EPurchaseTransactionState::Restored, TransactionData);
-	
-#if 0
-	bool bFound = false;
-	for (TSharedRef<FPurchaseReceipt>& OtherReceipt : OfflineTransactions)
-	{
-		// If redundant, replace the entry
-		if (OtherReceipt->TransactionId == OfflineReceipt->TransactionId)
-		{
-			*OtherReceipt = *OfflineReceipt;
-			bFound = true;
-			break;
-		}
-	}
-	
-	if (!bFound)
-#endif
-	{
-		OfflineTransactions.Add(OfflineReceipt);
-	}
+    FKnownTransaction* CachedTransaction = Transactions.FindByPredicate([&Transaction](const FKnownTransaction& KnownTransaction) {
+        return KnownTransaction.Receipt->TransactionId == Transaction.GetTransactionIdentifier();
+    });
+    if(CachedTransaction)
+    {
+        // In case the transaction already exists as Purchased don't make it Restored
+        if(CachedTransaction->Receipt->TransactionState == EPurchaseTransactionState::Purchased &&
+           Receipt->TransactionState == EPurchaseTransactionState::Restored)
+        {
+            return CachedTransaction->Receipt;
+        }
+        else
+        {
+            CachedTransaction->Receipt = Receipt;
+            return Receipt;
+        }
+    }
+    else
+    {
+        Transactions.Emplace(Receipt, Transaction.GetPaymentTransaction());
+        return Receipt;
+    }
 }
 
 void FOnlinePurchaseIOS::OnRestoreTransactionsComplete(EPurchaseTransactionState Result)
@@ -362,7 +625,7 @@ void FOnlinePurchaseIOS::OnRestoreTransactionsComplete(EPurchaseTransactionState
 	
 	// Full restore is complete
 	bRestoringTransactions = false;
-	bool bSuccess = (Result == EPurchaseTransactionState::Restored) || (Result == EPurchaseTransactionState::Purchased);
+	bool bSuccess = (Result == EPurchaseTransactionState::Restored);
 	Subsystem->ExecuteNextTick([this, bSuccess]() {
 		FOnlineError FinalResult(bSuccess);
 		QueryReceiptsComplete.ExecuteIfBound(FinalResult);
@@ -370,100 +633,39 @@ void FOnlinePurchaseIOS::OnRestoreTransactionsComplete(EPurchaseTransactionState
 	});
 }
 
-void FOnlinePurchaseIOS::OnTransactionInProgress(const FStoreKitTransactionData& TransactionData)
-{
-	UE_LOG_ONLINE_PURCHASE(Verbose, TEXT("FOnlinePurchaseIOS::OnTransactionInProgress TransactionData=[%s]"), *TransactionData.ToDebugString());
-}
-
-void FOnlinePurchaseIOS::OnTransactionDeferred(const FStoreKitTransactionData& TransactionData)
-{
-	UE_LOG_ONLINE_PURCHASE(Verbose, TEXT("FOnlinePurchaseIOS::OnTransactionDeferred TransactionData=[%s]"), *TransactionData.ToDebugString());
-	
-	FString UserIdStr = IOSUSER;
-	const TSharedRef<FOnlinePurchasePendingTransactionIOS>* UserPendingTransactionPtr = PendingTransactions.Find(UserIdStr);
-	if (UserPendingTransactionPtr != nullptr)
-	{
-		const TSharedRef<FOnlinePurchasePendingTransactionIOS> UserPendingTransaction = *UserPendingTransactionPtr;
-		
-		const FString& ErrorStr = TransactionData.GetErrorStr();
-		
-		FOnlineError FinalResult;
-		FinalResult.SetFromErrorCode(TEXT("com.epicgames.purchase.deferred"));
-		FinalResult.ErrorMessage = !ErrorStr.IsEmpty() ? FText::FromString(ErrorStr) : LOCTEXT("IOSTransactionDeferred", "Transaction awaiting approval.");
-		
-		TSharedRef<FPurchaseReceipt> DeferredReceipt = FOnlinePurchasePendingTransactionIOS::GenerateReceipt(EPurchaseTransactionState::Deferred, TransactionData);
-		
-		// Clear out the deferred transaction, it should get picked up in "offline" receipts
-		PendingTransactions.Remove(UserIdStr);
-		UserPendingTransaction->CheckoutCompleteDelegate.ExecuteIfBound(FinalResult, DeferredReceipt);
-
-#if TEST_DEFERRED_TRANSACTIONS
-		// Existing transactions emulating deferment will trigger as success shortly
-		TWeakPtr<FOnlinePurchaseIOS, ESPMode::ThreadSafe> WeakThis(AsShared());
-		FTSTicker::GetCoreTicker().AddTicker(FTickerDelegate::CreateLambda([WeakThis, TransactionData](float) -> bool
-		{
-			QUICK_SCOPE_CYCLE_COUNTER(STAT_FOnlinePurchaseIOS_TestDeferredTransactions);
-			FOnlinePurchaseIOSPtr StrongThis = WeakThis.Pin();
-			if (StrongThis.IsValid())
-			{
-				// Emulate redemption in a fixed time period
-				StrongThis->OnTransactionCompleteResponse(EPurchaseTransactionState::Purchased, TransactionData);
-			}
-			return false;
-		}), TEST_DEFERRED_DELAY);
-#endif // TEST_DEFERRED_TRANSACTIONS
-	}
-	else
-	{
-		UE_LOG_ONLINE_PURCHASE(Log, TEXT("Offline deferred transaction"));
-	}
-}
-
 void FOnlinePurchaseIOS::FinalizeReceiptValidationInfo(const FUniqueNetId& UserId, FString& InReceiptValidationInfo, const FOnFinalizeReceiptValidationInfoComplete& Delegate)
 {
 	FOnlineError DefaultSuccess(true);
 	Delegate.ExecuteIfBound(DefaultSuccess, InReceiptValidationInfo);
 }
 
-TSharedRef<FPurchaseReceipt> FOnlinePurchasePendingTransactionIOS::GenerateReceipt()
+TSharedRef<FPurchaseReceipt> FOnlinePurchaseIOS::GenerateFailReceipt(const FPurchaseCheckoutRequest& CheckoutRequest)
 {
-	TSharedRef<FPurchaseReceipt> Receipt = MakeShareable(new FPurchaseReceipt());
-	
-	Receipt->TransactionState = PendingPurchaseInfo.TransactionState;
-	Receipt->TransactionId = PendingPurchaseInfo.TransactionId;
-	
-	if(PendingPurchaseInfo.TransactionState == EPurchaseTransactionState::Purchased ||
-	   PendingPurchaseInfo.TransactionState == EPurchaseTransactionState::Restored)
-	{
-		Receipt->ReceiptOffers = PendingPurchaseInfo.ReceiptOffers;
-	}
-	else
-	{
-		// Add the requested offers to the receipt in the event of an incomplete purchase.
-		for(const auto& RequestedOffer : CheckoutRequest.PurchaseOffers)
-		{
-			Receipt->AddReceiptOffer(RequestedOffer.OfferNamespace, RequestedOffer.OfferId, RequestedOffer.Quantity);
-		}
-	}
-	
-	return Receipt;
+    TSharedRef<FPurchaseReceipt> Receipt = MakeShared<FPurchaseReceipt>();
+    
+    Receipt->TransactionState = EPurchaseTransactionState::Failed;
+    
+    // Add the requested offers to the receipt in the event of an incomplete purchase.
+    for(const auto& RequestedOffer : CheckoutRequest.PurchaseOffers)
+    {
+        Receipt->AddReceiptOffer(RequestedOffer.OfferNamespace, RequestedOffer.OfferId, RequestedOffer.Quantity);
+    }
+    return Receipt;
 }
 
-TSharedRef<FPurchaseReceipt> FOnlinePurchasePendingTransactionIOS::GenerateReceipt(EPurchaseTransactionState Result, const FStoreKitTransactionData& Transaction)
+TSharedRef<FPurchaseReceipt> FOnlinePurchaseIOS::GenerateOfflineReceipt(EPurchaseTransactionState Result, const FOnlinePurchaseTransactionIOS& Transaction)
 {
-	TSharedRef<FPurchaseReceipt> Receipt = MakeShareable(new FPurchaseReceipt());
-	
+    TSharedRef<FPurchaseReceipt> Receipt = MakeShared<FPurchaseReceipt>();
+
 	Receipt->TransactionState = Result;
 	Receipt->TransactionId = Transaction.GetTransactionIdentifier();
 	
 	if (Result == EPurchaseTransactionState::Purchased ||
 	    Result == EPurchaseTransactionState::Restored)
 	{
-		FPurchaseReceipt::FReceiptOfferEntry ReceiptEntry(TEXT(""), Transaction.GetOfferId(), 1);
-		
-		int32 Idx = ReceiptEntry.LineItems.AddZeroed();
+		FPurchaseReceipt::FReceiptOfferEntry ReceiptEntry(TEXT(""), Transaction.GetOfferId(), Transaction.GetQuantity());
 		
-		FPurchaseReceipt::FLineItemInfo& LineItem = ReceiptEntry.LineItems[Idx];
+		FPurchaseReceipt::FLineItemInfo& LineItem = ReceiptEntry.LineItems.Emplace_GetRef();
 		
 		LineItem.ItemName = Transaction.GetOfferId();
 		LineItem.UniqueId = Transaction.GetTransactionIdentifier();
@@ -475,83 +677,22 @@ TSharedRef<FPurchaseReceipt> FOnlinePurchasePendingTransactionIOS::GenerateRecei
 	return Receipt;
 }
 
-void FOnlinePurchasePendingTransactionIOS::StartProcessing()
-{
-	PendingPurchaseInfo.TransactionState = EPurchaseTransactionState::Processing;
-	for (int32 OfferIdx = 0; OfferIdx < CheckoutRequest.PurchaseOffers.Num(); ++OfferIdx)
-	{
-		OfferPurchaseStates[OfferIdx] = EPurchaseTransactionState::Processing;
-	}
-}
-
-bool FOnlinePurchasePendingTransactionIOS::AddCompletedOffer(EPurchaseTransactionState Result, const FStoreKitTransactionData& Transaction)
-{
-	for (int32 OfferIdx = 0; OfferIdx < CheckoutRequest.PurchaseOffers.Num(); ++OfferIdx)
-	{
-		const FPurchaseCheckoutRequest::FPurchaseOfferEntry& Offer = CheckoutRequest.PurchaseOffers[OfferIdx];
-		if (Transaction.GetOfferId() == Offer.OfferId)
-		{
-			OfferPurchaseStates[OfferIdx] = Result;
-			FPurchaseReceipt::FReceiptOfferEntry Receipt(TEXT(""), Transaction.GetOfferId(), 1);
-
-			int32 Idx = Receipt.LineItems.AddZeroed();
-			
-			FPurchaseReceipt::FLineItemInfo& LineItem = Receipt.LineItems[Idx];
-			
-			LineItem.ItemName = Transaction.GetOfferId();
-			LineItem.UniqueId = Transaction.GetTransactionIdentifier();
-			LineItem.ValidationInfo = Transaction.GetReceiptData();
-			
-			PendingPurchaseInfo.AddReceiptOffer(Receipt);
-			return true;
-		}
-	}
-	
-	return false;
-}
-
-bool FOnlinePurchasePendingTransactionIOS::AreAllOffersComplete() const
+TSharedRef<FPurchaseReceipt> FOnlinePurchaseIOS::GenerateReceipt(EPurchaseTransactionState Result, const FPurchaseCheckoutRequest& CheckoutRequest, const FOnlinePurchaseTransactionIOS& Transaction)
 {
-	for (const EPurchaseTransactionState State : OfferPurchaseStates)
-	{
-		if (State == EPurchaseTransactionState::NotStarted ||
-			State == EPurchaseTransactionState::Processing)
-		{
-			return false;
-		}
-	}
-	
-	return true;
-}
+    TSharedRef<FPurchaseReceipt> Receipt = GenerateOfflineReceipt(Result, Transaction);
 
-EPurchaseTransactionState FOnlinePurchasePendingTransactionIOS::GetFinalTransactionState() const
-{
-	bool bAnyFailures = false;
-	bool bAnyCancels = false;
-	for (const EPurchaseTransactionState State : OfferPurchaseStates)
+	if(Result != EPurchaseTransactionState::Purchased &&
+       Result != EPurchaseTransactionState::Restored)
 	{
-		if (State == EPurchaseTransactionState::NotStarted ||
-			State == EPurchaseTransactionState::Processing ||
-			State == EPurchaseTransactionState::Failed)
-		{
-			bAnyFailures = true;
-		}
-		else if (State == EPurchaseTransactionState::Canceled)
-		{
-			bAnyCancels = true;
-		}
-	}
-	
-	if (bAnyFailures)
-	{
-		return EPurchaseTransactionState::Failed;
-	}
-	else if (bAnyCancels)
-	{
-		return EPurchaseTransactionState::Canceled;
+        // Add the requested offers to the receipt in the event of an incomplete purchase.
+        for(const auto& RequestedOffer : CheckoutRequest.PurchaseOffers)
+        {
+            Receipt->AddReceiptOffer(RequestedOffer.OfferNamespace, RequestedOffer.OfferId, RequestedOffer.Quantity);
+        }
 	}
 	
-	return EPurchaseTransactionState::Purchased;
+	return Receipt;
 }
 
 #undef LOCTEXT_NAMESPACE
+
diff --git a/Engine/Plugins/Online/IOS/OnlineSubsystemIOS/Source/Private/OnlineStoreIOS.cpp b/Engine/Plugins/Online/IOS/OnlineSubsystemIOS/Source/Private/OnlineStoreIOS.cpp
--- a/Engine/Plugins/Online/IOS/OnlineSubsystemIOS/Source/Private/OnlineStoreIOS.cpp
+++ b/Engine/Plugins/Online/IOS/OnlineSubsystemIOS/Source/Private/OnlineStoreIOS.cpp
@@ -5,35 +5,126 @@
 #include "Internationalization/Culture.h"
 #include "Internationalization/FastDecimalFormat.h"
 #include "OnlineSubsystem.h"
-#import "OnlineStoreKitHelper.h"
+#import <StoreKit/SKProduct.h>
+#import <StoreKit/SKProductsRequest.h>
 #include "Misc/ConfigCacheIni.h"
 
-FOnlineStoreIOS::FOnlineStoreIOS(FOnlineSubsystemIOS* InSubsystem)
-	: StoreHelper(nullptr)
-	, bIsQueryInFlight(false)
-	, Subsystem(InSubsystem)
+@interface FSKProductsRequestHelper : SKProductsRequest
 {
-	UE_LOG_ONLINE_STOREV2(Verbose, TEXT( "FOnlineStoreIOS::FOnlineStoreIOS" ));
+};
+
+/** Delegate to fire when this product request completes with the store kit */
+@property FOnQueryOnlineStoreOffersComplete OfferDelegate;
+@end
+
+@implementation FSKProductsRequestHelper
+@end
+
+/**
+ * Proxy class that notifies updates from FSKProductsRequestHelper to FOnlineStoreIOS on the game 
+ * thread (SKRequest responses are received in worker threads)
+ */
+@interface FStoreKitStoreProxy : NSObject<SKProductsRequestDelegate>
+{
+    FOnlineStoreIOS* _StoreReceiver;
+    NSMutableSet* _Requests;
+};
+@end
+
+@implementation FStoreKitStoreProxy
+
+////////////////////////////////////////////////////////////////////
+/// SKProductsRequestDelegate implementation
+
+-(void)productsRequest: (SKProductsRequest *)Request didReceiveResponse: (SKProductsResponse *)Response
+{
+    UE_LOG_ONLINE_STOREV2(Log, TEXT("FStoreKitStoreProxy::didReceiveResponse"));
+
+    // Response for SKRequest is received in a working thread. Use a task to notify from game thread
+    FSKProductsRequestHelper* Helper = (FSKProductsRequestHelper*)Request;
+    [FIOSAsyncTask CreateTaskWithBlock : ^ bool(void)
+    {
+        if(_StoreReceiver)
+        {
+            _StoreReceiver->OnProductsRequestResponse(Response, Helper.OfferDelegate);
+        }
+        return true;
+    }];
 }
 
-FOnlineStoreIOS::FOnlineStoreIOS()
-	: bIsQueryInFlight(false)
+- (void)request:(SKRequest *)Request didFailWithError:(NSError *)Error
 {
-	UE_LOG_ONLINE_STOREV2(Verbose, TEXT( "FOnlineStoreIOS::FOnlineStoreIOS" ));
+    UE_LOG_ONLINE_STOREV2(Log, TEXT("FStoreKitStoreProxy::didFailWithError"));
+
+    // Response for SKRequest is received in a working thread. Use a task to notify from game thread
+    FSKProductsRequestHelper* Helper = (FSKProductsRequestHelper*)Request;
+    [FIOSAsyncTask CreateTaskWithBlock : ^ bool(void)
+    {
+        if(_StoreReceiver)
+        {
+            _StoreReceiver->OnProductsRequestResponse(nil, Helper.OfferDelegate);
+        }
+        return true;
+    }];}
+
+-(void)requestDidFinish:(SKRequest*)Request
+{
+    [_Requests removeObject:Request];
 }
 
-void FOnlineStoreIOS::InitStoreKit(FStoreKitHelperV2* InStoreKit)
+////////////////////////////////////////////////////////////////////
+/// FStoreKitStoreProxy methods
+
+- (id)initWithReceiver: (FOnlineStoreIOS*)StoreReceiver
 {
-	StoreHelper = InStoreKit;
-	
-	FOnProductsRequestResponseDelegate OnProductsRequestResponseDelegate;
-	OnProductsRequestResponseDelegate.BindRaw(this, &FOnlineStoreIOS::OnProductPurchaseRequestResponse);
-	[StoreHelper AddOnProductRequestResponse: OnProductsRequestResponseDelegate];
+    _StoreReceiver = StoreReceiver;
+    _Requests = [[NSMutableSet alloc] init];
+    return self;
+}
+
+-(void)dealloc
+{
+    for(FSKProductsRequestHelper* Request in _Requests)
+    {
+        [Request cancel];
+    }
+    [_Requests release];
+    
+    [super dealloc];
+}
+
+-(void)Shutdown
+{
+    _StoreReceiver = nullptr;
+}
+
+-(void)requestProductData: (NSMutableSet*)ProductIDs WithDelegate : (const FOnQueryOnlineStoreOffersComplete&)Delegate
+{
+    UE_LOG_ONLINE_STOREV2(Log, TEXT("FStoreKitStoreProxy::requestProductData"));
+
+    FSKProductsRequestHelper* Request = [[[FSKProductsRequestHelper alloc] initWithProductIdentifiers:ProductIDs] autorelease];
+    Request.delegate = self;
+    Request.OfferDelegate = Delegate;
+    [_Requests addObject:Request];
+    [Request start];
+}
+
+@end
+
+////////////////////////////////////////////////////////////////////
+/// FOnlineStoreIOS implementation
+
+FOnlineStoreIOS::FOnlineStoreIOS(FOnlineSubsystemIOS* InSubsystem)
+	: bIsQueryInFlight(false)
+	, Subsystem(InSubsystem)
+{
+	UE_LOG_ONLINE_STOREV2(Verbose, TEXT( "FOnlineStoreIOS::FOnlineStoreIOS" ));
+    StoreKitProxy = [[[FStoreKitStoreProxy alloc] initWithReceiver:this] autorelease];
 }
 
 FOnlineStoreIOS::~FOnlineStoreIOS()
 {
-	[StoreHelper release];
+    [StoreKitProxy Shutdown];
 }
 
 void FOnlineStoreIOS::QueryCategories(const FUniqueNetId& UserId, const FOnQueryOnlineStoreCategoriesComplete& Delegate)
@@ -66,14 +157,7 @@ bool FOnlineStoreIOS::OffersNotAllowedInLocale(const FString& InLocale)
     TArray<FString> LocaleData;
     InLocale.ParseIntoArray(LocaleData, TEXT("-"));
     FString Locale = LocaleData.Num() > 1 ? LocaleData[1] : LocaleData[0];
-    for (FString BannedLocale : BannedLocales)
-    {
-        if (BannedLocale == Locale)
-        {
-            return true;
-        }
-    }
-    return false;
+    return BannedLocales.Contains(Locale);
 }
 
 void FOnlineStoreIOS::QueryOffersById(const FUniqueNetId& UserId, const TArray<FUniqueOfferId>& OfferIds, const FOnQueryOnlineStoreOffersComplete& Delegate)
@@ -104,12 +188,7 @@ void FOnlineStoreIOS::QueryOffersById(const FUniqueNetId& UserId, const TArray<F
 			[ProductSet addObject:ID];
 		}
 		
-		FOnQueryOnlineStoreOffersComplete DelegateCopy = Delegate;
-		dispatch_async(dispatch_get_main_queue(), ^
-        {
-			[StoreHelper requestProductData:ProductSet WithDelegate:DelegateCopy];
-		});
-		
+        [StoreKitProxy requestProductData:ProductSet WithDelegate:Delegate];		
 		bIsQueryInFlight = true;
 	}
 }
@@ -122,103 +201,83 @@ void FOnlineStoreIOS::QueryOffersById(const FUniqueNetId& UserId, const TArray<F
  *
  * @return FOnlineStoreOffer with proper parameters filled in
  */
-FOnlineStoreOffer ConvertProductToStoreOffer(SKProduct* Product)
+TSharedPtr<FOnlineStoreOffer> ConvertProductToStoreOffer(SKProduct* Product)
 {
-	FOnlineStoreOffer NewProductInfo;
+    TSharedPtr<FOnlineStoreOffer> NewProductInfo = MakeShared<FOnlineStoreOffer>();
 	
-	NewProductInfo.OfferId = [Product productIdentifier];
+	NewProductInfo->OfferId = [Product productIdentifier];
 	
-	NewProductInfo.Title = FText::FromString([Product localizedTitle]);
-	NewProductInfo.Description = FText::FromString([Product localizedDescription]);
-	//NewProductInfo.LongDescription = FText::FromString([Product localizedDescription]);
-	NewProductInfo.CurrencyCode = [Product.priceLocale objectForKey : NSLocaleCurrencyCode];
+	NewProductInfo->Title = FText::FromString([Product localizedTitle]);
+	NewProductInfo->Description = FText::FromString([Product localizedDescription]);
+	//NewProductInfo->LongDescription = FText::FromString([Product localizedDescription]);
+	NewProductInfo->CurrencyCode = [Product.priceLocale objectForKey : NSLocaleCurrencyCode];
 	
 	// Convert the backend stated price into its base units
 	FInternationalization& I18N = FInternationalization::Get();
 	const FCulture& Culture = *I18N.GetCurrentCulture();
 	
-	const FDecimalNumberFormattingRules& FormattingRules = Culture.GetCurrencyFormattingRules(NewProductInfo.CurrencyCode);
+	const FDecimalNumberFormattingRules& FormattingRules = Culture.GetCurrencyFormattingRules(NewProductInfo->CurrencyCode);
 	const FNumberFormattingOptions& FormattingOptions = FormattingRules.CultureDefaultFormattingOptions;
 	double Val = static_cast<double>([Product.price doubleValue]) * static_cast<double>(FMath::Pow(10.0f, FormattingOptions.MaximumFractionalDigits));
-	NewProductInfo.NumericPrice = FMath::TruncToInt(Val + 0.5);
+	NewProductInfo->NumericPrice = FMath::TruncToInt(Val + 0.5);
 	
 	// iOS doesn't support these fields, set to min and max defaults
-	NewProductInfo.ReleaseDate = FDateTime::MinValue();
-	NewProductInfo.ExpirationDate = FDateTime::MaxValue();
+	NewProductInfo->ReleaseDate = FDateTime::MinValue();
+	NewProductInfo->ExpirationDate = FDateTime::MaxValue();
 
-	NewProductInfo.PriceText = FText::AsCurrencyBase(NewProductInfo.NumericPrice, NewProductInfo.CurrencyCode);
+	NewProductInfo->PriceText = FText::AsCurrencyBase(NewProductInfo->NumericPrice, NewProductInfo->CurrencyCode);
 	
 	return NewProductInfo;
 }
 
-void FOnlineStoreIOS::AddOffer(const FOnlineStoreOfferIOS& NewOffer)
+void FOnlineStoreIOS::OnProductsRequestResponse(SKProductsResponse* Response, const FOnQueryOnlineStoreOffersComplete& CompletionDelegate)
 {
-	if (NewOffer.IsValid())
-	{
-		FOnlineStoreOfferIOS* Existing = CachedOffers.Find(NewOffer.Offer->OfferId);
-		if (Existing != nullptr)
-		{
-			// Replace existing offer
-			*Existing = NewOffer;
-		}
-		else
-		{
-			CachedOffers.Add(NewOffer.Offer->OfferId, NewOffer);
-		}
-	}
+    if(bIsQueryInFlight)
+    {
+        TArray<FUniqueOfferId> OfferIds;
+        bool bWasSuccessful = (Response != nil);
+        if(bWasSuccessful)
+        {
+            for (SKProduct* Product in Response.products)
+            {
+                FOnlineStoreOfferIOS NewProductOffer(Product, ConvertProductToStoreOffer(Product));
+                
+                AddOffer(NewProductOffer);
+                OfferIds.Add(NewProductOffer.Offer->OfferId);
+                
+                UE_LOG_ONLINE_STOREV2(Log, TEXT("Product Identifier: %s, Name: %s, Desc: %s, Long Desc: %s, Price: %s IntPrice: %d"),
+                    *NewProductOffer.Offer->OfferId,
+                    *NewProductOffer.Offer->Title.ToString(),
+                    *NewProductOffer.Offer->Description.ToString(),
+                    *NewProductOffer.Offer->LongDescription.ToString(),
+                    *NewProductOffer.Offer->PriceText.ToString(),
+                    NewProductOffer.Offer->NumericPrice);
+            }
+            
+            for (NSString *invalidProduct in Response.invalidProductIdentifiers)
+            {
+                UE_LOG_ONLINE_STOREV2(Warning, TEXT("Problem in iTunes connect configuration for product: %s"), *FString(invalidProduct));
+            }
+        }        
+        CompletionDelegate.ExecuteIfBound(bWasSuccessful, OfferIds, TEXT(""));
+        bIsQueryInFlight = false;
+    }
 }
 
-void FOnlineStoreIOS::OnProductPurchaseRequestResponse(SKProductsResponse* Response, const FOnQueryOnlineStoreOffersComplete& CompletionDelegate)
+void FOnlineStoreIOS::AddOffer(const FOnlineStoreOfferIOS& NewOffer)
 {
-	if(bIsQueryInFlight)
+	if (NewOffer.IsValid())
 	{
-		bool bWasSuccessful = [Response.products count] > 0;
-		int32 NumInvalidProducts = [Response.invalidProductIdentifiers count];
-		if ([Response.products count] == 0 && NumInvalidProducts == 0)
-		{
-			UE_LOG_ONLINE_STOREV2(Warning, TEXT("Wrong number of products [%d] in the response when trying to make a single purchase"), [Response.products count]);
-		}
-		
-		TArray<FUniqueOfferId> OfferIds;
-		for (SKProduct* Product in Response.products)
-		{
-			FOnlineStoreOfferIOS NewProductOffer(Product, ConvertProductToStoreOffer(Product));
-			
-			AddOffer(NewProductOffer);
-			OfferIds.Add(NewProductOffer.Offer->OfferId);
-			
-			UE_LOG_ONLINE_STOREV2(Log, TEXT("Product Identifier: %s, Name: %s, Desc: %s, Long Desc: %s, Price: %s IntPrice: %d"),
-				   *NewProductOffer.Offer->OfferId,
-				   *NewProductOffer.Offer->Title.ToString(),
-				   *NewProductOffer.Offer->Description.ToString(),
-				   *NewProductOffer.Offer->LongDescription.ToString(),
-				   *NewProductOffer.Offer->PriceText.ToString(),
-				   NewProductOffer.Offer->NumericPrice);
-		}
-		
-		for (NSString *invalidProduct in Response.invalidProductIdentifiers)
-		{
-			UE_LOG_ONLINE_STOREV2(Warning, TEXT("Problem in iTunes connect configuration for product: %s"), *FString(invalidProduct));
-		}
-		
-		CompletionDelegate.ExecuteIfBound(bWasSuccessful, OfferIds, TEXT(""));
-		bIsQueryInFlight = false;
-	}
-}
-
-void FOnlineStoreIOS::OnTransactionCompleteResponse(EPurchaseTransactionState Result, const FStoreKitTransactionData& TransactionData)
-{
-	UE_LOG_ONLINE_STOREV2(Verbose, TEXT("FOnlineStoreIOS::OnTransactionCompleteResponse"));
-}
-
-void FOnlineStoreIOS::OnTransactionRestored(const FStoreKitTransactionData& TransactionData)
-{
-	UE_LOG_ONLINE_STOREV2(Verbose, TEXT("FOnlineStoreIOS::OnTransactionRestored"));
-}
-
-void FOnlineStoreIOS::OnRestoreTransactionsComplete(EPurchaseTransactionState Result)
-{
-	UE_LOG_ONLINE_STOREV2(Verbose, TEXT("FOnlineStoreIOS::OnRestoreTransactionsComplete"));	
+        FOnlineStoreOfferIOS* Existing = CachedOffers.Find(NewOffer.Offer->OfferId);
+        if (Existing != nullptr)
+        {
+            *Existing = NewOffer;
+        }
+        else
+        {
+            CachedOffers.Add(NewOffer.Offer->OfferId, NewOffer);
+        }
+    }
 }
 
 void FOnlineStoreIOS::GetOffers(TArray<FOnlineStoreOfferRef>& OutOffers) const
@@ -226,10 +285,7 @@ void FOnlineStoreIOS::GetOffers(TArray<FOnlineStoreOfferRef>& OutOffers) const
 	for (const auto& CachedEntry : CachedOffers)
 	{
 		const FOnlineStoreOfferIOS& CachedOffer = CachedEntry.Value;
-		if (CachedOffer.IsValid())
-		{
-			OutOffers.Add(CachedOffer.Offer.ToSharedRef());
-		}
+        OutOffers.Add(CachedOffer.Offer.ToSharedRef());
 	}
 }
 
@@ -238,7 +294,7 @@ TSharedPtr<FOnlineStoreOffer> FOnlineStoreIOS::GetOffer(const FUniqueOfferId& Of
 	TSharedPtr<FOnlineStoreOffer> Result;
 
 	const FOnlineStoreOfferIOS* Existing = CachedOffers.Find(OfferId);
-	if (Existing != nullptr && Existing->IsValid())
+	if (Existing != nullptr)
 	{
 		Result = Existing->Offer;
 	}
@@ -249,7 +305,7 @@ TSharedPtr<FOnlineStoreOffer> FOnlineStoreIOS::GetOffer(const FUniqueOfferId& Of
 SKProduct* FOnlineStoreIOS::GetSKProductByOfferId(const FUniqueOfferId& OfferId)
 {
 	const FOnlineStoreOfferIOS* Existing = CachedOffers.Find(OfferId);
-	if (Existing != nullptr && Existing->IsValid())
+	if (Existing != nullptr)
 	{
 		return Existing->Product;
 	}
diff --git a/Engine/Plugins/Online/IOS/OnlineSubsystemIOS/Source/Private/OnlineStoreKitHelper.cpp b/Engine/Plugins/Online/IOS/OnlineSubsystemIOS/Source/Private/OnlineStoreKitHelper.cpp
deleted file mode 100644
diff --git a/Engine/Plugins/Online/IOS/OnlineSubsystemIOS/Source/Private/OnlineStoreKitHelper.h b/Engine/Plugins/Online/IOS/OnlineSubsystemIOS/Source/Private/OnlineStoreKitHelper.h
deleted file mode 100644
diff --git a/Engine/Plugins/Online/IOS/OnlineSubsystemIOS/Source/Private/OnlineSubsystemIOS.cpp b/Engine/Plugins/Online/IOS/OnlineSubsystemIOS/Source/Private/OnlineSubsystemIOS.cpp
--- a/Engine/Plugins/Online/IOS/OnlineSubsystemIOS/Source/Private/OnlineSubsystemIOS.cpp
+++ b/Engine/Plugins/Online/IOS/OnlineSubsystemIOS/Source/Private/OnlineSubsystemIOS.cpp
@@ -2,17 +2,14 @@
 
 #include "OnlineSubsystemIOS.h"
 #include "IOS/IOSAppDelegate.h"
-#import "OnlineStoreKitHelper.h"
 #import "OnlineAppStoreUtils.h"
 #include "Misc/ConfigCacheIni.h"
 #include "Stats/Stats.h"
 
 FOnlineSubsystemIOS::FOnlineSubsystemIOS(FName InInstanceName)
 	: FOnlineSubsystemImpl(IOS_SUBSYSTEM, InInstanceName)
-	, StoreHelper(nil)
 	, AppStoreHelper(nil)
 {
-	StoreHelper = nil;
 }
 
 IOnlineSessionPtr FOnlineSubsystemIOS::GetSessionInterface() const
@@ -171,7 +168,6 @@ bool FOnlineSubsystemIOS::Init()
 		{
 			StoreV2Interface = MakeShareable(new FOnlineStoreIOS(this));
 			PurchaseInterface = MakeShareable(new FOnlinePurchaseIOS(this));
-			InitStoreKitHelper();
 		}
 
 		if (UserCloudInterface && IsCloudKitEnabled())
@@ -191,27 +187,9 @@ bool FOnlineSubsystemIOS::Init()
 	return bSuccessfullyStartedUp;
 }
 
-void FOnlineSubsystemIOS::InitStoreKitHelper()
-{
-	StoreHelper = [[FStoreKitHelperV2 alloc] init];
-	[StoreHelper retain];
-	
-	// Give each interface a chance to bind to the store kit helper
-	StoreV2Interface->InitStoreKit(StoreHelper);
-	PurchaseInterface->InitStoreKit(StoreHelper);
-
-}
-
-void FOnlineSubsystemIOS::CleanupStoreKitHelper()
-{
-	[StoreHelper release];
-	StoreHelper = nil;
-}
-
 void FOnlineSubsystemIOS::InitAppStoreHelper()
 {
 	AppStoreHelper = [[FAppStoreUtils alloc] init];
-	[AppStoreHelper retain];
 }
 
 void FOnlineSubsystemIOS::CleanupAppStoreHelper()
@@ -276,7 +254,6 @@ bool FOnlineSubsystemIOS::Shutdown()
 #undef DESTRUCT_INTERFACE
 	
 	// Cleanup after the interfaces are free
-	CleanupStoreKitHelper();
 	CleanupAppStoreHelper();
 	
 	return bSuccessfullyShutdown;
@@ -311,7 +288,7 @@ bool FOnlineSubsystemIOS::HandlePurchaseExecCommands(UWorld* InWorld, const TCHA
 
 	if (FParse::Command(&Cmd, TEXT("DUMPAPPRECEIPT")))
 	{
-		[StoreHelper dumpAppReceipt];
+        PurchaseInterface->DumpAppReceipt();
 	}
 
 	return bWasHandled;
diff --git a/Engine/Plugins/Online/IOS/OnlineSubsystemIOS/Source/Private/RetainedObjCInstance.h b/Engine/Plugins/Online/IOS/OnlineSubsystemIOS/Source/Private/RetainedObjCInstance.h
new file mode 100644
--- /dev/null
+++ b/Engine/Plugins/Online/IOS/OnlineSubsystemIOS/Source/Private/RetainedObjCInstance.h
@@ -0,0 +1,150 @@
+// Copyright Epic Games, Inc. All Rights Reserved.
+
+#pragma once
+
+#include <type_traits>
+
+/**
+ * RAII wrapper type that encapsulates manual memory management for Objective-C instances:
+ * - It calls retain when a non nil instance is received
+ * - It calls release on destruction
+ * - It handles copy/move semantics to keep reference counting consistent
+ * This type is useful because ARC is disabled for several reasons
+ */
+
+template <typename T>
+class TRetainedObjCInstance
+{
+    static_assert(std::is_pointer_v<T>, "T must be an Objective-C pointer type");
+public:
+    /**
+	 * Default constructor, initializes to nil.
+	 */
+    TRetainedObjCInstance()
+    :Ptr(nil)
+    {
+    }
+
+    /**
+	 * Constructor, initializes instance and increases reference count if not nil.
+	 */
+    TRetainedObjCInstance(T InPtr)
+    :Ptr(InPtr?[InPtr retain]:nil)
+    {
+    }
+
+    /**
+	 * Copy constructor, initializes instance and increases reference count if not nil.
+     * 	
+     * @param Other The Other object to copy from.
+	 */
+    TRetainedObjCInstance(const TRetainedObjCInstance& Other)
+    :TRetainedObjCInstance(Other.Ptr)
+    {
+    }
+    
+    /**
+	 * Move constructor, initializes instance. Sets other to nil without updating reference counting.
+     * 	
+     * @param Other The Other object to move from.
+	 */
+    TRetainedObjCInstance(TRetainedObjCInstance&& Other)
+    :Ptr(Other.Ptr)
+    {
+        Other.Ptr = nil;
+    }
+    
+    /**
+	 * Copy assignment, releases owned instance and retains received one.
+     * 	
+     * @param Other The Other object to copy from.
+	 */
+    TRetainedObjCInstance& operator=(const TRetainedObjCInstance& Other)
+    {
+        if(this != &Other)
+        {
+            [Ptr release];
+            Ptr = [Other.Ptr retain];
+        }
+        return *this;
+    }
+    
+    /**
+	 * Move assignment, releases owned instance, updates owned instance and clears Other.
+     * 	
+     * @param Other The Other object to move from.
+	 */
+    TRetainedObjCInstance operator=(TRetainedObjCInstance&& Other)
+    {
+        if(this != &Other)
+        {
+            [Ptr release];
+            Ptr = Other.Ptr;
+            Other.Ptr = nil;
+        }
+        return *this;
+    }
+
+    /**
+	 * Destructor, releases owned instance if needed
+	 */
+    ~TRetainedObjCInstance()
+    {
+        Reset();
+    }
+    
+    /**
+	 * Conversion operator to underlying instance. Useful to invoke Objective-C methods transparently
+	 */    
+    operator T() { return Ptr; }
+
+    /**
+	 * Conversion operator to underlying instance. Useful to invoke Objective-C methods transparently
+	 */    
+    operator T() const { return Ptr; }
+
+    /**
+	 * Releases ownership on Objective-C instance
+	 */    
+    void Reset()
+    {
+        if(Ptr)
+        {
+            [Ptr release];
+            Ptr = nil;
+        }
+    }
+    
+    friend bool operator==(const TRetainedObjCInstance<T>& Lhs, const TRetainedObjCInstance<T>& Rhs)
+    {
+        return Lhs.Ptr == Rhs.Ptr;
+    }
+
+    friend bool operator==(const TRetainedObjCInstance<T>& Lhs, T Rhs)
+    {
+        return Lhs.Ptr == Rhs;
+    }
+
+    friend bool operator==(T Lhs,const TRetainedObjCInstance<T>& Rhs)
+    {
+        return (Rhs == Lhs);
+    }
+
+    friend bool operator!=(const TRetainedObjCInstance<T>& Lhs, const TRetainedObjCInstance<T>& Rhs)
+    {
+        return Lhs.Ptr != Rhs.Ptr;
+    }
+
+    friend bool operator!=(const TRetainedObjCInstance<T>& Lhs, T Rhs)
+    {
+        return Lhs.Ptr != Rhs;
+    }
+
+    friend bool operator!=(T Lhs, const TRetainedObjCInstance<T>& Rhs)
+    {
+        return (Rhs != Lhs);
+    }
+private:
+    /* Owned Objective-C instance */
+    T Ptr;
+};
diff --git a/Engine/Plugins/Online/IOS/OnlineSubsystemIOS/Source/Public/OnlinePurchaseIOS.h b/Engine/Plugins/Online/IOS/OnlineSubsystemIOS/Source/Public/OnlinePurchaseIOS.h
--- a/Engine/Plugins/Online/IOS/OnlineSubsystemIOS/Source/Public/OnlinePurchaseIOS.h
+++ b/Engine/Plugins/Online/IOS/OnlineSubsystemIOS/Source/Public/OnlinePurchaseIOS.h
@@ -3,64 +3,13 @@
 #pragma once
 
 #include "Interfaces/OnlinePurchaseInterface.h"
-#include "OnlineSubsystemIOSTypes.h"
+#include "RetainedObjCInstance.h"
 
-struct FStoreKitTransactionData;
-@class SKProductsResponse;
-@class FStoreKitHelperV2;
-
-/**
- * Info used to cache and track orders in progress.
- */
-class FOnlinePurchasePendingTransactionIOS
-{
-public:
-	FOnlinePurchasePendingTransactionIOS(
-		const FPurchaseCheckoutRequest& InCheckoutRequest,
-		const FUniqueNetId& InUserId,
-		const EPurchaseTransactionState InPendingTransactionState = EPurchaseTransactionState::NotStarted,
-		const FOnPurchaseCheckoutComplete& InCheckoutCompleteDelegate = FOnPurchaseCheckoutComplete()
-		)
-		: CheckoutRequest(InCheckoutRequest)
-		, UserId(StaticCastSharedRef<const FUniqueNetIdIOS>(InUserId.AsShared()))
-		, CheckoutCompleteDelegate(InCheckoutCompleteDelegate)
-	{
-		PendingPurchaseInfo.TransactionState = InPendingTransactionState;
-		
-		// Setup purchase state for all pending offers
-		OfferPurchaseStates.AddZeroed(CheckoutRequest.PurchaseOffers.Num());
-	}
-	
-	/**
-	 * Generate a final receipt for all purchases made in this single transaction
-	 */
-	TSharedRef<FPurchaseReceipt> GenerateReceipt();
-	
-	/** Generate one off receipts for transactions initiated outside the current run of the application */
-	static TSharedRef<FPurchaseReceipt> GenerateReceipt(EPurchaseTransactionState Result, const FStoreKitTransactionData& Transaction);
-	
-	/** Mark this pending purchase as started */
-	void StartProcessing();
-	/** Add a single completed transaction (one of possibly many) to this transaction */
-	bool AddCompletedOffer(EPurchaseTransactionState Result, const FStoreKitTransactionData& Transaction);
-	/** @return true if all offers to purchase have completed transactions */
-	bool AreAllOffersComplete() const;
-	/** @return result of all purchases, if any fail it will return failed, if any cancel it will return canceled */
-	EPurchaseTransactionState GetFinalTransactionState() const;
-	
-public:
-	
-	/** Checkout info for the pending order */
-	const FPurchaseCheckoutRequest CheckoutRequest;
-	/** Mirror array of purchase state for the various offers to purchase */
-	TArray<EPurchaseTransactionState> OfferPurchaseStates;
-	/** User for the pending order */
-	const FUniqueNetIdIOSRef UserId;
-	/** Delegate to call on completion */
-	const FOnPurchaseCheckoutComplete CheckoutCompleteDelegate;
-	/** Tracks the current state of the order */
-	FPurchaseReceipt PendingPurchaseInfo;
-};
+enum class EPurchaseTransactionState : uint8;
+struct FOnlinePurchaseTransactionIOS;
+struct FOnlinePurchaseInProgressTransactionIOS;
+@class SKPaymentTransaction;
+@class FStoreKitPurchaseProxy;
 
 /**
  * Implementation for online purchase via IOS services
@@ -92,62 +41,73 @@ public:
 	 */
 	FOnlinePurchaseIOS(class FOnlineSubsystemIOS* InSubsystem);
 
-	/**
-	 * Constructor
-	 */
-	FOnlinePurchaseIOS();
-
 	/**
 	 * Destructor
 	 */
 	virtual ~FOnlinePurchaseIOS();
 	
-	/** Initialize the FStoreKitHelper for interaction with the app store */
-	void InitStoreKit(FStoreKitHelperV2* InStoreKit);
-	
-private:
-	
-	/** delegate fired when a product request completes */
-	void OnProductPurchaseRequestResponse(SKProductsResponse* Response, const FOnQueryOnlineStoreOffersComplete& CompletionDelegate);
-	/** delegate fired when a single purchase transaction has completed (may be a part of multiple requests at once) */
-	void OnTransactionCompleteResponse(EPurchaseTransactionState Result, const FStoreKitTransactionData& TransactionData);
-	/** delegate fired when a single transaction is restored (may be a part of many restored purchases) */
-	void OnTransactionRestored(const FStoreKitTransactionData& TransactionData);
-	/** delegate fired when all transactions have been restored */
-	void OnRestoreTransactionsComplete(EPurchaseTransactionState Result);
-	
-	void OnTransactionInProgress(const FStoreKitTransactionData& TransactionData);
-	void OnTransactionDeferred(const FStoreKitTransactionData& TransactionData);
-	
+	/**
+	 * Log App Receipt content
+	 */
+	void DumpAppReceipt();
+
+    /**
+     * Method used internally by FStoreKitPurchaseProxy to notify that a transaction was completed. Not meant to be called by user code
+     */
+	void OnTransactionComplete(EPurchaseTransactionState Result, const FOnlinePurchaseTransactionIOS& TransactionData);
+
+    /**
+     * Method used internally by FStoreKitPurchaseProxy to notify that transactions were restored. Not meant to be called by user code
+     */
+    void OnRestoreTransactionsComplete(EPurchaseTransactionState Result);
 private:
-	
-	/** Mapping from user id to pending transaction */
-	typedef TMap<FString, TSharedRef<FOnlinePurchasePendingTransactionIOS> > FOnlinePurchasePendingTransactionMap;
+    /**
+     * Add Transaction data to cache avoiding duplicates and return proper receipt
+     */
+    TSharedRef<FPurchaseReceipt> TryStoreTransactionAndGetReceipt(const TSharedRef<FPurchaseReceipt>& Receipt, const FOnlinePurchaseTransactionIOS& Transaction);
 
-	/** Mapping from user id to complete transactions */
-	typedef TMap<FString, TArray< TSharedRef<FPurchaseReceipt> > > FOnlinePurchaseCompleteTransactionsMap;
-	
-	/** Array of transactions completion indirectly (previous run, etc) */
-	typedef TArray< TSharedRef<FPurchaseReceipt> > FOnlineCompletedTransactions;
-	
+    /**
+     * Generate a receipt for a failed purchase
+     */
+    static TSharedRef<FPurchaseReceipt> GenerateFailReceipt(const FPurchaseCheckoutRequest& CheckoutRequest);
+    
+    /**
+     * Generate a receipt for a successful purchase which was received offline
+     */
+    static TSharedRef<FPurchaseReceipt> GenerateOfflineReceipt(EPurchaseTransactionState Result, const FOnlinePurchaseTransactionIOS& Transaction);
+
+    /**
+     * Generate a final receipt for the purchase currently in progress
+     */
+    static TSharedRef<FPurchaseReceipt> GenerateReceipt(EPurchaseTransactionState Result, const FPurchaseCheckoutRequest& CheckoutRequest, const FOnlinePurchaseTransactionIOS& Transaction);
 private:
 	
-	/** Store kit helper for interfacing with app store (owned by main OnlineSubsystem) */
-	FStoreKitHelperV2* StoreHelper;
+	/** Proxy helper to communicate from/to StoreKit*/
+    TRetainedObjCInstance<FStoreKitPurchaseProxy*> StoreKitProxy;
 	
 	/** Are transactions current being restored */
 	bool bRestoringTransactions;
-	/** Transient delegate to fire when query receipts has completed, when restoring transactions */
+
+	/** Transient delegate to fire when restoring transactions (QueryReceipts with bRestoreReceipts=true)*/
 	FOnQueryReceiptsComplete QueryReceiptsComplete;
 	
-	/** Keeps track of pending user transactions */
-	FOnlinePurchasePendingTransactionMap PendingTransactions;
-	
+	/** Keeps track of in progress user transaction */
+    TSharedPtr<const FOnlinePurchaseInProgressTransactionIOS> InProgressTransaction;
+	
+	/** Entry data for a known completed transaction (restored or purchased) */
+    struct FKnownTransaction
+    {
+        explicit FKnownTransaction(TSharedRef<FPurchaseReceipt> InReceipt, SKPaymentTransaction* InTransaction)
+        : Receipt(InReceipt)
+        , PaymentTransaction(InTransaction)
+        {
+        }
+
+        TSharedRef<FPurchaseReceipt> Receipt;
+        TRetainedObjCInstance<SKPaymentTransaction*> PaymentTransaction;
+    };
 	/** Cache of completed transactions */
-	FOnlinePurchaseCompleteTransactionsMap CompletedTransactions;
-	
-	/** Cache of purchases completed outside the running instance */
-	FOnlineCompletedTransactions OfflineTransactions;
+    TArray< FKnownTransaction > Transactions;
 	
 	/** Reference to the parent subsystem */
 	FOnlineSubsystemIOS* Subsystem;
diff --git a/Engine/Plugins/Online/IOS/OnlineSubsystemIOS/Source/Public/OnlineStoreIOS.h b/Engine/Plugins/Online/IOS/OnlineSubsystemIOS/Source/Public/OnlineStoreIOS.h
--- a/Engine/Plugins/Online/IOS/OnlineSubsystemIOS/Source/Public/OnlineStoreIOS.h
+++ b/Engine/Plugins/Online/IOS/OnlineSubsystemIOS/Source/Public/OnlineStoreIOS.h
@@ -3,10 +3,10 @@
 #pragma once
 
 #include "Interfaces/OnlineStoreInterfaceV2.h"
+#include "RetainedObjCInstance.h"
 
 enum class EPurchaseTransactionState : uint8;
-struct FStoreKitTransactionData;
-@class FStoreKitHelperV2;
+@class FStoreKitStoreProxy;
 @class SKProductsResponse;
 @class SKProduct;
 
@@ -20,7 +20,7 @@ class FOnlineStoreIOS :
 public:
 
 	// IOnlineStoreV2
-
+    
 	virtual void QueryCategories(const FUniqueNetId& UserId, const FOnQueryOnlineStoreCategoriesComplete& Delegate) override;
 	virtual void GetCategories(TArray<FOnlineStoreCategory>& OutCategories) const override;
 	virtual void QueryOffersByFilter(const FUniqueNetId& UserId, const FOnlineStoreFilter& Filter, const FOnQueryOnlineStoreOffersComplete& Delegate) override;
@@ -37,70 +37,40 @@ public:
 	 */
 	FOnlineStoreIOS(class FOnlineSubsystemIOS* InSubsystem);
 
-	/**
-	 * Constructor
-	 */
-	FOnlineStoreIOS();
-
 	/**
 	 * Destructor
 	 */
 	virtual ~FOnlineStoreIOS();
 	
-	/** Initialize the FStoreKitHelper for interaction with the app store */
-	void InitStoreKit(FStoreKitHelperV2* InStoreKit);
-	
 	/**
 	 * Get the production information for a given offer id
-	 * Must have previously been retrieved vi QueryOffers*
+	 * Must have previously been retrieved by QueryOffers*
 	 *
 	 * @return the SKProduct previously queried for product information
 	 */
 	SKProduct* GetSKProductByOfferId(const FUniqueOfferId& OfferId);
 
+    /**
+     * Method used internally by FOnlineSubsystemIOS to notify that QueryOffers request finished. Not meant to be called by user code
+     */
+    void OnProductsRequestResponse(SKProductsResponse* Response, const FOnQueryOnlineStoreOffersComplete& CompletionDelegate);
+
 private:
 	
 	/**
 	 * Representation of a single product offer
 	 */
-	struct FOnlineStoreOfferIOS
+    struct FOnlineStoreOfferIOS
 	{
 		FOnlineStoreOfferIOS()
-		: Product(nil)
-		, Offer(nullptr)
+		: FOnlineStoreOfferIOS(nil, nullptr)
 		{
 		}
 		
-		FOnlineStoreOfferIOS(SKProduct* InProduct, const FOnlineStoreOffer& InOffer)
+		FOnlineStoreOfferIOS(SKProduct* InProduct, TSharedPtr<FOnlineStoreOffer>&& InOffer)
 		{
 			Product = InProduct;
-			[Product retain];
-			
-			Offer = MakeShareable(new FOnlineStoreOffer(InOffer));
-		}
-		
-		FOnlineStoreOfferIOS(const FOnlineStoreOfferIOS& Other)
-		{
-			Product = Other.Product;
-			[Product retain];
-			Offer = Other.Offer;
-		}
-		
-		FOnlineStoreOfferIOS& operator=(const FOnlineStoreOfferIOS& Other)
-		{
-			// Free
-			[Product release];
-
-			// Copy
-			Product = Other.Product;
-			[Product retain];
-			Offer = Other.Offer;
-			return *this;
-		}
-		
-		~FOnlineStoreOfferIOS()
-		{
-			[Product release];
+			Offer = MoveTemp(InOffer);
 		}
 		
 		/** @return true if the store offer is valid/proper */
@@ -110,7 +80,7 @@ private:
 		}
 		
 		/** Reference to the app store product information */
-		SKProduct* Product;
+		TRetainedObjCInstance<SKProduct*> Product;
 		/** Product information about this offer */
 		TSharedPtr<FOnlineStoreOffer> Offer;
 	};
@@ -119,17 +89,6 @@ private:
 
     bool OffersNotAllowedInLocale(const FString& Locale);
     
-private:
-
-	/** delegate fired when a product request completes */
-	void OnProductPurchaseRequestResponse(SKProductsResponse* Response, const FOnQueryOnlineStoreOffersComplete& CompletionDelegate);
-	/** delegate fired when a single purchase transaction has completed (may be a part of multiple requests at once) */
-	void OnTransactionCompleteResponse(EPurchaseTransactionState Result, const FStoreKitTransactionData& TransactionData);
-	/** delegate fired when a single transaction is restored (may be a part of many restored purchases) */
-	void OnTransactionRestored(const FStoreKitTransactionData& TransactionData);
-	/** delegate fired when all transactions have been restored */
-	void OnRestoreTransactionsComplete(EPurchaseTransactionState Result);
-	
 private:
 	
 	/** Mapping from offer id to product information */
@@ -139,7 +98,7 @@ private:
 	FOnlineOfferDescriptionMap CachedOffers;
 
 	/** Store kit helper for interfacing with app store */
-	FStoreKitHelperV2* StoreHelper;
+    TRetainedObjCInstance<FStoreKitStoreProxy*> StoreKitProxy;
 
 	/** Is a query already in flight */
 	bool bIsQueryInFlight;
diff --git a/Engine/Plugins/Online/IOS/OnlineSubsystemIOS/Source/Public/OnlineSubsystemIOS.h b/Engine/Plugins/Online/IOS/OnlineSubsystemIOS/Source/Public/OnlineSubsystemIOS.h
--- a/Engine/Plugins/Online/IOS/OnlineSubsystemIOS/Source/Public/OnlineSubsystemIOS.h
+++ b/Engine/Plugins/Online/IOS/OnlineSubsystemIOS/Source/Public/OnlineSubsystemIOS.h
@@ -17,7 +17,6 @@
 #include "OnlineUserCloudInterfaceIOS.h"
 #include "OnlineSharedCloudInterfaceIOS.h"
 
-@class FStoreKitHelperV2;
 @class FAppStoreUtils;
 @class GKPlayer;
 @class GKLocalPlayer;
@@ -104,9 +103,6 @@ public:
 
 private:
 	
-	void InitStoreKitHelper();
-	void CleanupStoreKitHelper();
-	
 	void InitAppStoreHelper();
 	void CleanupAppStoreHelper();
 
@@ -116,9 +112,6 @@ private:
 	/** Online async task thread */
 	class FRunnableThread* OnlineAsyncTaskThread;
 	
-	/** Store kit helper for interfacing with app store */
-	FStoreKitHelperV2* StoreHelper;
-	
 	/** App store util helper */
 	FAppStoreUtils* AppStoreHelper;
 
diff --git a/Engine/Source/Runtime/ApplicationCore/Private/IOS/IOSAppDelegate.cpp b/Engine/Source/Runtime/ApplicationCore/Private/IOS/IOSAppDelegate.cpp
--- a/Engine/Source/Runtime/ApplicationCore/Private/IOS/IOSAppDelegate.cpp
+++ b/Engine/Source/Runtime/ApplicationCore/Private/IOS/IOSAppDelegate.cpp
@@ -20,7 +20,7 @@
 #include "IOS/IOSAsyncTask.h"
 #include "Misc/ConfigCacheIni.h"
 #include "IOS/IOSPlatformCrashContext.h"
-#include "IOS/IOSPaymentTransactionObserver.h"
+#include "IOS/ProxyPaymentTransactionObserver.h"
 #include "Misc/OutputDeviceError.h"
 #include "Misc/OutputDeviceRedirector.h"
 #include "Misc/FeedbackContext.h"
@@ -1087,7 +1087,7 @@ static FAutoConsoleVariableRef CVarGEnableThermalsReport(
 	[[NSNotificationCenter defaultCenter] addObserver:self selector:@selector(OnVoiceOverStatusChanged) name:UIAccessibilityVoiceOverStatusDidChangeNotification object:nil];
 #endif
 
-	[[SKPaymentQueue defaultQueue] addTransactionObserver:[FPaymentTransactionObserver sharedInstance]];
+	[[SKPaymentQueue defaultQueue] addTransactionObserver:[FProxyPaymentTransactionObserver sharedInstance]];
     
 	return YES;
 }
@@ -1477,6 +1477,8 @@ extern double GCStartTime;
 	 */
 	FCoreDelegates::ApplicationWillTerminateDelegate.Broadcast();
     
+    [[SKPaymentQueue defaultQueue] removeTransactionObserver:[FProxyPaymentTransactionObserver sharedInstance]];
+
     // note that we are shutting down
     // TODO: fix the reason why we are hanging when asked to shutdown
 /*    RequestEngineExit(TEXT("IOS applicationWillTerminate"));
diff --git a/Engine/Source/Runtime/ApplicationCore/Private/IOS/IOSPaymentTransactionObserver.cpp b/Engine/Source/Runtime/ApplicationCore/Private/IOS/IOSPaymentTransactionObserver.cpp
deleted file mode 100644
diff --git a/Engine/Source/Runtime/ApplicationCore/Private/IOS/ProxyPaymentTransactionObserver.cpp b/Engine/Source/Runtime/ApplicationCore/Private/IOS/ProxyPaymentTransactionObserver.cpp
new file mode 100644
--- /dev/null
+++ b/Engine/Source/Runtime/ApplicationCore/Private/IOS/ProxyPaymentTransactionObserver.cpp
@@ -0,0 +1,59 @@
+// Copyright Epic Games, Inc. All Rights Reserved.
+
+#include "IOS/ProxyPaymentTransactionObserver.h"
+
+@implementation FProxyPaymentTransactionObserver
+
++ (FProxyPaymentTransactionObserver*) sharedInstance
+{
+    static dispatch_once_t onceToken;
+    static FProxyPaymentTransactionObserver* transactionObserverSharedInstance;
+
+    dispatch_once(&onceToken, ^{
+        transactionObserverSharedInstance = [[FProxyPaymentTransactionObserver alloc] init];
+    });
+    return transactionObserverSharedInstance;
+}
+
+- (id) init
+{
+	self = [super init];
+	return self;
+}
+
+-(void)paymentQueue:(SKPaymentQueue*)queue updatedTransactions : (NSArray*)transactions
+{
+    if(self.ossObserver)
+    {
+        [self.ossObserver paymentQueue:queue updatedTransactions:transactions];
+    }
+}
+
+/// Logs all transactions that have been removed from the payment queue.
+-(void)paymentQueue:(SKPaymentQueue*)queue removedTransactions : (NSArray*)transactions
+{
+    if(self.ossObserver)
+    {
+        [self.ossObserver paymentQueue:queue removedTransactions:transactions];
+    }
+}
+
+/// Called when an error occur while restoring purchases. Notifies the user about the error.
+-(void)paymentQueue:(SKPaymentQueue*)queue restoreCompletedTransactionsFailedWithError : (NSError*)error
+{
+    if(self.ossObserver)
+    {
+        [self.ossObserver paymentQueue:queue restoreCompletedTransactionsFailedWithError:error];
+    }
+}
+
+/// Called when all restorable transactions have been processed by the payment queue.
+-(void)paymentQueueRestoreCompletedTransactionsFinished:(SKPaymentQueue*)queue
+{
+    if(self.ossObserver)
+    {
+        [self.ossObserver paymentQueueRestoreCompletedTransactionsFinished:queue];
+    }
+}
+
+@end
diff --git a/Engine/Source/Runtime/ApplicationCore/Public/IOS/IOSPaymentTransactionObserver.h b/Engine/Source/Runtime/ApplicationCore/Public/IOS/IOSPaymentTransactionObserver.h
deleted file mode 100644
diff --git a/Engine/Source/Runtime/ApplicationCore/Public/IOS/ProxyPaymentTransactionObserver.h b/Engine/Source/Runtime/ApplicationCore/Public/IOS/ProxyPaymentTransactionObserver.h
new file mode 100644
--- /dev/null
+++ b/Engine/Source/Runtime/ApplicationCore/Public/IOS/ProxyPaymentTransactionObserver.h
@@ -0,0 +1,22 @@
+// Copyright Epic Games, Inc. All Rights Reserved.
+
+#pragma once
+
+#import "StoreKit/StoreKit.h"
+
+/**
+ * The proxy SKPaymentTransactionObserver implementation that receives transaction events raised by StoreKit. This
+ * class enables us to bind an observer on initialization in application:didFinishLaunchingWithOptions so 
+ * SKPaymentsQueue is properly synchronized with the App Store
+ * The consuming OSS starts up later and binds itself as the nested observer to receive notifications
+ */
+/** */
+@interface FProxyPaymentTransactionObserver : NSObject <SKPaymentTransactionObserver>
+{
+};
+
+@property (atomic, retain) id <SKPaymentTransactionObserver> ossObserver;
+
++ (FProxyPaymentTransactionObserver*) sharedInstance;
+
+@end
