Add missing include

diff --git a/Engine/Plugins/Online/IOS/OnlineSubsystemIOS/Source/Private/OnlinePurchaseIOS.cpp b/Engine/Plugins/Online/IOS/OnlineSubsystemIOS/Source/Private/OnlinePurchaseIOS.cpp
--- a/Engine/Plugins/Online/IOS/OnlineSubsystemIOS/Source/Private/OnlinePurchaseIOS.cpp
+++ b/Engine/Plugins/Online/IOS/OnlineSubsystemIOS/Source/Private/OnlinePurchaseIOS.cpp
@@ -3,6 +3,7 @@
 #include "OnlinePurchaseIOS.h"
 #include "OnlineError.h"
 #include "OnlineSubsystemIOS.h"
+#include "Misc/CommandLine.h"
 #include "Stats/Stats.h"
 #include "IOS/ProxyPaymentTransactionObserver.h"
 
