Fix iOS GameCenter login flow on iOS to fix Lyra issues

Setting GKLocalPlayer authenticateHandler only shows a native login UI/greeting the first time it is set. After that it receives GameCenter login state updates. Previously looks like it was assuming a login flow was initiated from that. Once GameCenter reports an authentication failure error no more errors are notified through the handler so it was possible that OnLoginComplete delegates were not called after a first failure anymore

It is possible tha a user moves the app to background and changes the GameCenter user from device settings, so this changes also adds triggering OnLoginStatusChanged delegates that were missing

diff --git a/Engine/Plugins/Online/IOS/OnlineSubsystemIOS/Source/Private/OnlineExternalUIInterfaceIOS.cpp b/Engine/Plugins/Online/IOS/OnlineSubsystemIOS/Source/Private/OnlineExternalUIInterfaceIOS.cpp
--- a/Engine/Plugins/Online/IOS/OnlineSubsystemIOS/Source/Private/OnlineExternalUIInterfaceIOS.cpp
+++ b/Engine/Plugins/Online/IOS/OnlineSubsystemIOS/Source/Private/OnlineExternalUIInterfaceIOS.cpp
@@ -17,16 +17,9 @@ bool FOnlineExternalUIIOS::ShowLoginUI(const int ControllerIndex, bool bShowOnli
 	
 	check(IdentityInterface != nullptr);
 	
-	if ([GKLocalPlayer localPlayer] == nil)
+	if (FUniqueNetIdPtr UniqueNetId = IdentityInterface->GetUniquePlayerId(ControllerIndex))
 	{
-		UE_LOG_ONLINE_EXTERNALUI(Log, TEXT("Game Center localPlayer is null."));
-		Delegate.ExecuteIfBound(nullptr, ControllerIndex, FOnlineError(EOnlineErrorResult::Unknown));
-		return true;
-	}
-
-	if ([GKLocalPlayer localPlayer].isAuthenticated)
-	{
-		Delegate.ExecuteIfBound(IdentityInterface->GetLocalPlayerUniqueId(), ControllerIndex, FOnlineError::Success());
+		Delegate.ExecuteIfBound(UniqueNetId, ControllerIndex, FOnlineError::Success());
 		return true;
 	}
 	
diff --git a/Engine/Plugins/Online/IOS/OnlineSubsystemIOS/Source/Private/OnlineIdentityInterfaceIOS.cpp b/Engine/Plugins/Online/IOS/OnlineSubsystemIOS/Source/Private/OnlineIdentityInterfaceIOS.cpp
--- a/Engine/Plugins/Online/IOS/OnlineSubsystemIOS/Source/Private/OnlineIdentityInterfaceIOS.cpp
+++ b/Engine/Plugins/Online/IOS/OnlineSubsystemIOS/Source/Private/OnlineIdentityInterfaceIOS.cpp
@@ -5,6 +5,8 @@
 #include "OnlineSubsystemIOS.h"
 #include "OnlineError.h"
 #include "IOS/IOSAppDelegate.h"
+
+#import <GameKit/GameKit.h>
 #import "OnlineAppStoreUtils.h"
 
 FOnlineIdentityIOS::FOnlineIdentityIOS()
@@ -18,16 +20,6 @@ FOnlineIdentityIOS::FOnlineIdentityIOS(FOnlineSubsystemIOS* InSubsystem)
 {
 }
 
-FUniqueNetIdIOSPtr FOnlineIdentityIOS::GetLocalPlayerUniqueId() const
-{
-	return UniqueNetId;
-}
-
-void FOnlineIdentityIOS::SetLocalPlayerUniqueId(const FUniqueNetIdIOSPtr& UniqueId)
-{
-	UniqueNetId = UniqueId;
-}
-
 TSharedPtr<FUserOnlineAccount> FOnlineIdentityIOS::GetUserAccount(const FUniqueNetId& UserId) const
 {
 	// not implemented
@@ -42,147 +34,175 @@ TArray<TSharedPtr<FUserOnlineAccount> > FOnlineIdentityIOS::GetAllUserAccounts()
 	return Result;
 }
 
-// Make sure IOnlineIdentity::Login is not called directly on iOS. Please use IOnlineExternalUI::ShowLoginUI instead.
 bool FOnlineIdentityIOS::Login(int32 LocalUserNum, const FOnlineAccountCredentials& AccountCredentials)
 {
-	bool bStartedLogin = false;
-
-	// Since the iOS login code may show a UI, ShowLoginUI is a better fit here. Also, note that the ConnectToService blueprint
-	// node that calls Login is deprecated (there's a new ShowExternalLoginUI node meant to replace it).
-
-	GKLocalPlayer* GKLocalUser = [GKLocalPlayer localPlayer];
+	// FOnlineIdentityIOS does not support more than 1 local player
+	if (LocalUserNum != 0 || bLoginInProgress)
+	{
+		return false;
+	}
 
-	// Was the login handled by Game Center
-	if( GKLocalUser && 
-		GKLocalUser.isAuthenticated )
+	bLoginInProgress = true;
+	
+	// All accesses to authenticateHandler should be done through main thread because it is non atomic
+	dispatch_async(dispatch_get_main_queue(), ^
 	{
-		// Now logged in
-		bStartedLogin = true;
-		if ([GKPlayer respondsToSelector:@selector(scopedIDsArePersistent)] == YES)
+		// Game Center state cannot be checked until authenticateHandler is set. Also, when setting authenticateHandler the Game Center
+		// authentication process starts and a login screen may be shown to the user if the device is not logged into Game Center.
+		// This is the reason it is set on the first login call instead of in an Init method or in the constructor
+		// There is no way to show this login screen again from the app if the user discards it
+		// If the device was logged into Game Center it will be shown a welcome message that will fade out after a few seconds
+		// Once the authenticateHandler is set all Game Center status changes will be reported through it including first authentication
+		// result. That is, if the device was already logged into Game Center when we set the authenticateHandler we will receive a notification
+		// on it
+		// The user can move to background at anytime and logout or change identity from device settings.
+		// Logout events are reported as errors in the handler so will have a non nil Error
+		// Login events are reported as non errors so will have a nill Error and [GKLocalPlayer localPlayer].isAuthenticated == YES
+		if ([GKLocalPlayer localPlayer].authenticateHandler == nil)
 		{
-			if ([GKLocalUser scopedIDsArePersistent])
+			[GKLocalPlayer localPlayer].authenticateHandler = (^(UIViewController* ViewController, NSError *Error)
 			{
-				const FString PlayerId(FString(FOnlineSubsystemIOS::GetPlayerId(GKLocalUser)));
+				if (ViewController == nil)
+				{
+					// Game Center notified a change of state. Identify what happened and handle the change from game thread
+					FGamCenterEvent Event = GetCurrentGameCenterEvent(Error);
+					
+					[FIOSAsyncTask CreateTaskWithBlock : ^ bool(void)
+					{
+						HandleGamCenterEvent(Event, bLoginInProgress);
+						bLoginInProgress = false;
+						return true;
+					}];
+				}
+				else
+				{
+					// Game Center has provided a view controller for us to fulfill some needed actions, present it from main thread.
+					UE_LOG_ONLINE_IDENTITY(Log, TEXT("Showing Game Center's provided UI for additional actions"));
+					dispatch_async(dispatch_get_main_queue(), ^
+					{
+						[[IOSAppDelegate GetDelegate].IOSController presentViewController:ViewController animated:YES completion:nil];
+					});
+				}
+			});
+		}
+		else
+		{
+			// If a login request is made after authenticateHandler is set just notify the current state
+			// Game Center state updates will always be done through authenticateHandler
+			[FIOSAsyncTask CreateTaskWithBlock : ^ bool(void)
+			{
+				if (UniqueNetId.IsValid())
+				{
+					TriggerOnLoginCompleteDelegates(0, true, *UniqueNetId, TEXT(""));
+				}
+				else
+				{
+					TriggerOnLoginCompleteDelegates(0, false, *FUniqueNetIdIOS::EmptyId(), TEXT("Game Center authentication failed. Please, try login to Game Center from the iOS device settings"));
+				}
+				bLoginInProgress = false;
+				return true;
+			}];
+		}
+	});
 
-				UniqueNetId = FUniqueNetIdIOS::Create(PlayerId);
-				TriggerOnLoginCompleteDelegates(LocalUserNum, true, *UniqueNetId, TEXT(""));
+	return true;
+}
 
-				UE_LOG_ONLINE_IDENTITY(Log, TEXT("The user %s has logged into Game Center"), *PlayerId);
-			}
-			else
-			{
-				// ID is not persistent across multiple game sessions, consider as not logged in
-				FString ErrorMessage = TEXT("The user could not be authenticated with a persistent id by Game Center");
-				UE_LOG_ONLINE_IDENTITY(Log, TEXT("%s"), *ErrorMessage);
+bool FOnlineIdentityIOS::Logout(int32 LocalUserNum)
+{
+	// There is no way to logout from Game Center from inside an app
+	// Logout is only possible by closing session from Game Center's iOS device settings
 
-				TriggerOnLoginCompleteDelegates(LocalUserNum, false, *FUniqueNetIdIOS::EmptyId(), *ErrorMessage);
-			}
-		}
-		else
+	return false;
+}
+
+void FOnlineIdentityIOS::HandleGamCenterEvent(const FGamCenterEvent& Event, bool bTriggerLoginComplete)
+{
+	check(Event.HasValue() || Event.HasError()); // Just check nobody stole the value/error unintentionally in the future
+	
+	// Clear UniqueNetId but keep it locally
+	FUniqueNetIdIOSPtr PreviousUniqueNetId = MoveTemp(UniqueNetId);
+	bool bWasLoggedIn = PreviousUniqueNetId.IsValid();
+
+	if (Event.HasValue())
+	{
+		UniqueNetId = Event.GetValue();
+		
+		if (bTriggerLoginComplete)
 		{
-			const FString PlayerId(FString(FOnlineSubsystemIOS::GetPlayerId(GKLocalUser)));
+			TriggerOnLoginCompleteDelegates(0, true, *UniqueNetId, TEXT(""));
+		}
 
-			UniqueNetId = FUniqueNetIdIOS::Create( PlayerId );
-			TriggerOnLoginCompleteDelegates(LocalUserNum, true, *UniqueNetId, TEXT(""));
+		if (!bWasLoggedIn || (*PreviousUniqueNetId != *UniqueNetId))
+		{
+			if (bWasLoggedIn)
+			{
+				UE_LOG_ONLINE_IDENTITY(Log, TEXT("Game Center device session closed for user %s"), *PreviousUniqueNetId->UniqueNetIdStr);
+			}
+			UE_LOG_ONLINE_IDENTITY(Log, TEXT("Game Center device session started for user %s"), *UniqueNetId->UniqueNetIdStr);
 
-			UE_LOG_ONLINE_IDENTITY(Log, TEXT("The user %s has logged into Game Center"), *PlayerId);
+			TriggerOnLoginStatusChangedDelegates(0, bWasLoggedIn? ELoginStatus::LoggedIn : ELoginStatus::NotLoggedIn, ELoginStatus::LoggedIn, *UniqueNetId);
+			TriggerOnLoginChangedDelegates(0);
 		}
 	}
 	else
 	{
-		// Trigger the login event on the main thread.
-		bStartedLogin = true;
-		dispatch_async(dispatch_get_main_queue(), ^
+		UE_LOG_ONLINE_IDENTITY(Log, TEXT("%s"), *Event.GetError());
+		
+		if (bTriggerLoginComplete)
 		{
-			[[GKLocalPlayer localPlayer] setAuthenticateHandler:(^(UIViewController* viewcontroller, NSError *error)
-			{
-				// The login process has completed.
-				if (viewcontroller == nil)
-				{
-					bool bWasSuccessful = false;
-					FString ErrorMessage;
-					TOptional<FString> PlayerId;
-
-					if (error)
-					{
-						// We did not complete authentication without error, we are not logged in - Game Center is not available
-						NSString *errstr = [error localizedDescription];
-						UE_LOG_ONLINE_IDENTITY(Warning, TEXT("Game Center login has failed. %s]"), *FString(errstr));
-
-						ErrorMessage = TEXT("An error occured authenticating the user by Game Center");
-						UE_LOG_ONLINE_IDENTITY(Log, TEXT("%s"), *ErrorMessage);
-					}
-					else if ([GKLocalPlayer localPlayer].isAuthenticated == YES)
-					{
-						GKLocalPlayer* GKLocalUserAuth = [GKLocalPlayer localPlayer];
+			TriggerOnLoginCompleteDelegates(0, false, *FUniqueNetIdIOS::EmptyId(), Event.GetError());
+		}
 
-						/* Perform additional tasks for the authenticated player here */
-						if ([GKPlayer respondsToSelector:@selector(scopedIDsArePersistent)] == YES)
-						{
-							if ([GKLocalUserAuth scopedIDsArePersistent])
-							{
-								PlayerId = FString(FOnlineSubsystemIOS::GetPlayerId(GKLocalUserAuth));
+		if (bWasLoggedIn)
+		{
+			UE_LOG_ONLINE_IDENTITY(Log, TEXT("Game Center device session closed for user %s"), *PreviousUniqueNetId->UniqueNetIdStr);
+			TriggerOnLoginStatusChangedDelegates(0, ELoginStatus::LoggedIn, ELoginStatus::NotLoggedIn, *FUniqueNetIdIOS::EmptyId());
+			TriggerOnLoginChangedDelegates(0);
+		}
+	}
+}
 
-								bWasSuccessful = true;
-								UE_LOG_ONLINE_IDENTITY(Log, TEXT("The user %s has logged into Game Center"), *PlayerId.GetValue());
-							}
-							else
-							{
-								// ID is not persistent across multiple game sessions, consider as not logged in
-								ErrorMessage = TEXT("The user could not be authenticated with a persistent id by Game Center");
-								UE_LOG_ONLINE_IDENTITY(Log, TEXT("%s"), *ErrorMessage);
-							}
-						}
-						else
-						{
-							PlayerId = FString(FOnlineSubsystemIOS::GetPlayerId(GKLocalUserAuth));
+FOnlineIdentityIOS::FGamCenterEvent FOnlineIdentityIOS::GetCurrentGameCenterEvent(NSError* Error)
+{
+	if (Error)
+	{
+		// Game Center notified an error. user is logged out from Game Center or there was an issue while authenticating
+		NSString *errstr = [Error localizedDescription];
+		UE_LOG_ONLINE_IDENTITY(Warning, TEXT("Game Center reported an error: %s]"), *FString(errstr));
+		
+		return MakeError(TEXT("Game Center reported an authentication error"));
+	}
+	else if ([GKLocalPlayer localPlayer].isAuthenticated == YES)
+	{
+		// Game Center authentication succeeded. Check if we can use the identifier
+		GKLocalPlayer* GKLocalUser = [GKLocalPlayer localPlayer];
 
-							bWasSuccessful = true;
-							UE_LOG_ONLINE_IDENTITY(Log, TEXT("The user %s has logged into Game Center"), *PlayerId.GetValue());
-						}
-					}
-					else
-					{
-						ErrorMessage = TEXT("The user could not be authenticated by Game Center");
-						UE_LOG_ONLINE_IDENTITY(Log, TEXT("%s"), *ErrorMessage);
-					}
+		if ([GKPlayer respondsToSelector:@selector(scopedIDsArePersistent)] == YES)
+		{
+			// ID is not persistent across multiple game sessions, consider as not logged in
+			return MakeError(TEXT("The user could not be authenticated with a persistent id by Game Center"));
+		}
 
-					// Report back to the game thread whether this succeeded.
-					[FIOSAsyncTask CreateTaskWithBlock : ^ bool(void)
-					{
-						if (PlayerId.IsSet())
-						{
-							UniqueNetId = FUniqueNetIdIOS::Create(PlayerId.GetValue());
-						}
-						else
-						{
-							UniqueNetId.Reset();
-						}
+		NSString* GameCenterPlayerId = FOnlineSubsystemIOS::GetPlayerId(GKLocalUser);
+		if ([GameCenterPlayerId isEqualToString: @"UnknownID"])
+		{
+			// Spuriously we may receive "UnknownID" as the id. In those cases ignore this value as a valid id and wait for the good one
+			// This has been noticed when moving the app to background, logout and login from Game Center as a different account and move the app to foregroung
+			// A good id will be received in the authenticateHandler on a subsequent event
+			return MakeError(TEXT("UnknownID player id received. If a proper id is not received on following login attempt try logout and login from Game Center from the iOS device settings"));
+		}
 
-						FUniqueNetIdIOSRef UniqueIdForUser = UniqueNetId.IsValid() ? UniqueNetId.ToSharedRef() : FUniqueNetIdIOS::EmptyId();
-						TriggerOnLoginCompleteDelegates(LocalUserNum, bWasSuccessful, *UniqueIdForUser, *ErrorMessage);
+		FString PlayerId(GameCenterPlayerId);
+		UE_LOG_ONLINE_IDENTITY(Log, TEXT("Active user for Game Center is %s"), *PlayerId);
 
-						return true;
-					}];
-				}
-				else
-				{
-					// Game Center has provided a view controller for us to login, we present it.
-					[[IOSAppDelegate GetDelegate].IOSController
-						presentViewController:viewcontroller animated:YES completion:nil];
-				}
-			})];
-		});
+		return MakeValue(FUniqueNetIdIOS::Create(MoveTemp(PlayerId)));
+	}
+	else
+	{
+		// According to documentation this should not happen
+		return MakeError(TEXT("The user could not be authenticated by Game Center"));
 	}
-	
-	return bStartedLogin;
-}
-
-bool FOnlineIdentityIOS::Logout(int32 LocalUserNum)
-{
-	UniqueNetId.Reset();
-	TriggerOnLogoutCompleteDelegates(LocalUserNum, false);
-	return true;
 }
 
 bool FOnlineIdentityIOS::AutoLogin(int32 LocalUserNum)
@@ -192,36 +212,34 @@ bool FOnlineIdentityIOS::AutoLogin(int32 LocalUserNum)
 
 ELoginStatus::Type FOnlineIdentityIOS::GetLoginStatus(int32 LocalUserNum) const
 {
-	ELoginStatus::Type LoginStatus = ELoginStatus::NotLoggedIn;
-
-	if(LocalUserNum < MAX_LOCAL_PLAYERS && GetLocalGameCenterUser() != NULL && GetLocalGameCenterUser().isAuthenticated == YES)
+	// FOnlineIdentityIOS does not support more than 1 local player
+	if (LocalUserNum != 0)
 	{
-		LoginStatus = ELoginStatus::LoggedIn;
+		return ELoginStatus::NotLoggedIn;
 	}
-
-	return LoginStatus;
+	
+	return UniqueNetId.IsValid()? ELoginStatus::LoggedIn : ELoginStatus::NotLoggedIn;
 }
 
 ELoginStatus::Type FOnlineIdentityIOS::GetLoginStatus(const FUniqueNetId& UserId) const 
 {
-	ELoginStatus::Type LoginStatus = ELoginStatus::NotLoggedIn;
-
-	if(GetLocalGameCenterUser() != NULL && GetLocalGameCenterUser().isAuthenticated == YES)
-	{
-		LoginStatus = ELoginStatus::LoggedIn;
-	}
-
-	return LoginStatus;
+	return UniqueNetId.IsValid() && UserId == *UniqueNetId? ELoginStatus::LoggedIn : ELoginStatus::NotLoggedIn;
 }
 
 FUniqueNetIdPtr FOnlineIdentityIOS::GetUniquePlayerId(int32 LocalUserNum) const
 {
+	// FOnlineIdentityIOS does not support more than 1 local player
+	if (LocalUserNum != 0)
+	{
+		return nullptr;
+	}
+	
 	return UniqueNetId;
 }
 
 FUniqueNetIdPtr FOnlineIdentityIOS::CreateUniquePlayerId(uint8* Bytes, int32 Size)
 {
-	if( Bytes && Size == sizeof(uint64) )
+	if (Bytes && Size == sizeof(uint64))
 	{
 		int32 StrLen = FCString::Strlen((TCHAR*)Bytes);
 		if (StrLen > 0)
@@ -241,9 +259,17 @@ FUniqueNetIdPtr FOnlineIdentityIOS::CreateUniquePlayerId(const FString& Str)
 
 FString FOnlineIdentityIOS::GetPlayerNickname(int32 LocalUserNum) const
 {
-	if (LocalUserNum < MAX_LOCAL_PLAYERS && GetLocalGameCenterUser() != NULL)
+	// FOnlineIdentityIOS does not support more than 1 local player
+	if (LocalUserNum != 0)
 	{
-		NSString* PersonaName = [GetLocalGameCenterUser() alias];
+		return FString();
+	}
+	
+	GKLocalPlayer* LocalPlayer = GetLocalGameCenterUser();
+	
+	if (LocalPlayer != nil)
+	{
+		NSString* PersonaName = [LocalPlayer alias];
 		
 		if (PersonaName != nil)
 		{
@@ -256,14 +282,9 @@ FString FOnlineIdentityIOS::GetPlayerNickname(int32 LocalUserNum) const
 
 FString FOnlineIdentityIOS::GetPlayerNickname(const FUniqueNetId& UserId) const 
 {
-	if (GetLocalGameCenterUser() != NULL)
+	if (UniqueNetId.IsValid() && UserId == *UniqueNetId)
 	{
-		NSString* PersonaName = [GetLocalGameCenterUser() alias];
-		
-		if (PersonaName != nil)
-		{
-			return FString(PersonaName);
-		}
+		return GetPlayerNickname(0);
 	}
 
 	return FString();
@@ -374,13 +395,9 @@ void FOnlineIdentityIOS::GetUserPrivilege(const FUniqueNetId& UserId, EUserPrivi
 
 FPlatformUserId FOnlineIdentityIOS::GetPlatformUserIdFromUniqueNetId(const FUniqueNetId& InUniqueNetId) const
 {
-	for (int i = 0; i < MAX_LOCAL_PLAYERS; ++i)
+	if (UniqueNetId.IsValid() && (*UniqueNetId == InUniqueNetId))
 	{
-		auto CurrentUniqueId = GetUniquePlayerId(i);
-		if (CurrentUniqueId.IsValid() && (*CurrentUniqueId == InUniqueNetId))
-		{
-			return GetPlatformUserIdFromLocalUserNum(i);
-		}
+		return GetPlatformUserIdFromLocalUserNum(0);
 	}
 
 	return PLATFORMUSERID_NONE;
@@ -390,3 +407,8 @@ FString FOnlineIdentityIOS::GetAuthType() const
 {
 	return TEXT("");
 }
+
+GKLocalPlayer* FOnlineIdentityIOS::GetLocalGameCenterUser() const
+{
+	return UniqueNetId.IsValid() ? [GKLocalPlayer localPlayer] : nil;
+}
diff --git a/Engine/Plugins/Online/IOS/OnlineSubsystemIOS/Source/Private/OnlineIdentityInterfaceIOS.h b/Engine/Plugins/Online/IOS/OnlineSubsystemIOS/Source/Private/OnlineIdentityInterfaceIOS.h
--- a/Engine/Plugins/Online/IOS/OnlineSubsystemIOS/Source/Private/OnlineIdentityInterfaceIOS.h
+++ b/Engine/Plugins/Online/IOS/OnlineSubsystemIOS/Source/Private/OnlineIdentityInterfaceIOS.h
@@ -3,8 +3,8 @@
 
 #pragma once
 
-#import <UIKit/UIKit.h>
-#import <GameKit/GameKit.h>
+@class GKLocalPlayer;
+@class NSError;
 
 #include "OnlineSubsystemIOSTypes.h"
 #include "Interfaces/OnlineIdentityInterface.h"
@@ -17,9 +17,17 @@ private:
 	/** UID for this identity */
 	FUniqueNetIdIOSPtr UniqueNetId;
 	FOnlineSubsystemIOS* Subsystem;
-
+    
+    static_assert(MAX_LOCAL_PLAYERS == 1, "FOnlineIdentityIOS does not support more than 1 local player");
+    bool bLoginInProgress = false;
+    
 	FOnlineIdentityIOS();
 
+    using FGamCenterEvent = TValueOrError<FUniqueNetIdIOSPtr, FString>;
+    
+    static FGamCenterEvent GetCurrentGameCenterEvent(NSError* Error);
+    void HandleGamCenterEvent(const FGamCenterEvent& LoginResult, bool bTriggerLoginComplete);
+    
 PACKAGE_SCOPE:
 
 	/**
@@ -27,10 +35,6 @@ PACKAGE_SCOPE:
 	 */
 	FOnlineIdentityIOS(FOnlineSubsystemIOS* InSubsystem);
 
-	FUniqueNetIdIOSPtr GetLocalPlayerUniqueId() const;
-
-	void SetLocalPlayerUniqueId(const FUniqueNetIdIOSPtr& UniqueId);
-
 public:
 
 	//~ Begin IOnlineIdentity Interface
@@ -66,10 +70,7 @@ public:
 	 *
 	 * @return - The game center local player
 	 */
-	GKLocalPlayer* GetLocalGameCenterUser() const
-	{
-		return UniqueNetId.IsValid() ? [GKLocalPlayer localPlayer] : nil;
-	}
+    GKLocalPlayer* GetLocalGameCenterUser() const;
 };
 
 
