Add missing header to fix iOS NoPCH build

diff --git a/Engine/Plugins/Online/IOS/OnlineSubsystemIOS/Source/Private/OnlineIdentityInterfaceIOS.h b/Engine/Plugins/Online/IOS/OnlineSubsystemIOS/Source/Private/OnlineIdentityInterfaceIOS.h
--- a/Engine/Plugins/Online/IOS/OnlineSubsystemIOS/Source/Private/OnlineIdentityInterfaceIOS.h
+++ b/Engine/Plugins/Online/IOS/OnlineSubsystemIOS/Source/Private/OnlineIdentityInterfaceIOS.h
@@ -8,6 +8,7 @@
 
 #include "OnlineSubsystemIOSTypes.h"
 #include "Interfaces/OnlineIdentityInterface.h"
+#include "Templates/ValueOrError.h"
 #include "OnlineSubsystemIOSPackage.h"
 
 class FOnlineIdentityIOS :
