iOS: fix a bug when CloudKit is not yet initialized and trying to read from it immediately

diff --git a/Engine/Plugins/Online/IOS/OnlineSubsystemIOS/Source/Private/OnlineUserCloudInterfaceIOS.cpp b/Engine/Plugins/Online/IOS/OnlineSubsystemIOS/Source/Private/OnlineUserCloudInterfaceIOS.cpp
--- a/Engine/Plugins/Online/IOS/OnlineSubsystemIOS/Source/Private/OnlineUserCloudInterfaceIOS.cpp
+++ b/Engine/Plugins/Online/IOS/OnlineSubsystemIOS/Source/Private/OnlineUserCloudInterfaceIOS.cpp
@@ -623,13 +623,35 @@ bool FOnlineUserCloudInterfaceIOS::OnReadUserCloudFileBegin(const FString & File
 	{
 		OnReadUserCloudFileCompleteDelegateHandle = AddOnReadUserFileCompleteDelegate_Handle(OnReadUserCloudFileCompleteDelegate);
 
-		ReadUserFile(*UniqueNetId, FileName);
-
-		FCloudFile* CloudFile = GetCloudFile(FileName, false);
+		// check cloudkit was properly initialized
+		while (MetaDataState == EOnlineAsyncTaskState::InProgress)
+		{
+			FPlatformProcess::Sleep(0.01f);	// Cloudkit still needs to finish initialization
+		}
+		if (MetaDataState == EOnlineAsyncTaskState::Failed)
+		{
+			return false; // Cloudkit has failed to initialize for some reason
+		}
+		ensureMsgf(MetaDataState == EOnlineAsyncTaskState::Done, "CloudKit should be initialized first");
 
-		while (CloudFile && CloudFile->AsyncState == EOnlineAsyncTaskState::InProgress)
+		// check named record in Cloudkit
+		if (!ReadUserFile(*UniqueNetId, FileName))
 		{
-			FPlatformProcess::Sleep(0.01f);
+			return false; // ReadUserFile failed
+		}
+		bool bLoop = true;
+		while (bLoop)
+		{
+			// loop until the async task set in ReadUserFile is finished (Done or Failed)
+			{
+				FScopeLock ScopeLock(&CloudDataLock);
+				FCloudFile* CloudFile = GetCloudFile(FileName, false);
+				bLoop = (CloudFile) ? (CloudFile->AsyncState == EOnlineAsyncTaskState::InProgress) : false;
+			}
+			if (bLoop)
+			{
+				FPlatformProcess::Sleep(0.01f);
+			}
 		}
 
 		// clean up temporary data
