iOS: Fix Compilation Tests / NonUnity Compile UnrealGame
add missing TEXT("...") around ensureMsgf string

diff --git a/Engine/Plugins/Online/IOS/OnlineSubsystemIOS/Source/Private/OnlineUserCloudInterfaceIOS.cpp b/Engine/Plugins/Online/IOS/OnlineSubsystemIOS/Source/Private/OnlineUserCloudInterfaceIOS.cpp
--- a/Engine/Plugins/Online/IOS/OnlineSubsystemIOS/Source/Private/OnlineUserCloudInterfaceIOS.cpp
+++ b/Engine/Plugins/Online/IOS/OnlineSubsystemIOS/Source/Private/OnlineUserCloudInterfaceIOS.cpp
@@ -632,7 +632,7 @@ bool FOnlineUserCloudInterfaceIOS::OnReadUserCloudFileBegin(const FString & File
 		{
 			return false; // Cloudkit has failed to initialize for some reason
 		}
-		ensureMsgf(MetaDataState == EOnlineAsyncTaskState::Done, "CloudKit should be initialized first");
+		ensureMsgf(MetaDataState == EOnlineAsyncTaskState::Done, TEXT("CloudKit should be initialized first"));
 
 		// check named record in Cloudkit
 		if (!ReadUserFile(*UniqueNetId, FileName))
