Fix Receipt management on OSSiOS
Receipts retrieved using GetReceipts should be updated by calls to QueryReceipts. After a successful Checkout the receipt for the successful purchase is added to the cached receipts
Ignored all asynchronous transaction updates received. Just trigger the UnexpectedPurchaseReceipt delegate in that case of Purchased transaction notifications
Restored transactions are automatically finished so the end user does not need to do it

Subscriptions and durable products can be purchased and validated but there is no way to know about their state on device using StoreKit. I tried to decode the iOS AppStore receipt using OpenSSL but there is no useful information about owned products (contents look equivalent to restoring receipts)

Also, used proper log category for the purchase interface

diff --git a/Engine/Plugins/Online/IOS/OnlineSubsystemIOS/Source/Private/OnlinePurchaseIOS.cpp b/Engine/Plugins/Online/IOS/OnlineSubsystemIOS/Source/Private/OnlinePurchaseIOS.cpp
--- a/Engine/Plugins/Online/IOS/OnlineSubsystemIOS/Source/Private/OnlinePurchaseIOS.cpp
+++ b/Engine/Plugins/Online/IOS/OnlineSubsystemIOS/Source/Private/OnlinePurchaseIOS.cpp
@@ -5,7 +5,6 @@
 #include "OnlineSubsystemIOS.h"
 #include "Misc/CommandLine.h"
 #include "Stats/Stats.h"
-#include "IOS/ProxyPaymentTransactionObserver.h"
 
 #import <StoreKit/SKError.h>
 #import <StoreKit/SKReceiptRefreshRequest.h>
@@ -28,23 +27,35 @@ namespace OnlinePurchaseIOSPrivate
      *
      * @return hex encoded string with opaque data representing a completed transaction
      */
-    FString ConvertReceiptToString()
+
+	FString ConvertReceiptToString()
+	{
+		FString ReceiptData;
+		
+		NSURL* nsReceiptUrl = [[NSBundle mainBundle] appStoreReceiptURL];
+		NSData* nsReceiptData = [NSData dataWithContentsOfURL : nsReceiptUrl];
+		if (nsReceiptData)
+		{
+			NSString* nsEncodedReceiptData = [nsReceiptData base64EncodedStringWithOptions: 0];
+			ReceiptData = nsEncodedReceiptData;
+			UE_LOG_ONLINE_PURCHASE(VeryVerbose, TEXT("ConvertReceiptToString %s"), *ReceiptData);
+		}
+		else
+		{
+			UE_LOG_ONLINE_PURCHASE(Log, TEXT("ConvertReceiptToString: No receipt data found for transaction"));
+		}
+		return ReceiptData;
+	}
+
+	TSharedPtr<FString> ConvertReceiptToSharedString()
     {
-        FString ReceiptData;
-        
-        NSURL* nsReceiptUrl = [[NSBundle mainBundle] appStoreReceiptURL];
-        NSData* nsReceiptData = [NSData dataWithContentsOfURL : nsReceiptUrl];
-        if (nsReceiptData)
-        {
-            NSString* nsEncodedReceiptData = [nsReceiptData base64EncodedStringWithOptions : NSDataBase64EncodingEndLineWithLineFeed];
-            ReceiptData = nsEncodedReceiptData;
-            UE_LOG_ONLINE_STOREV2(VeryVerbose, TEXT("ConvertReceiptToString %s"), *ReceiptData);
-        }
-        else
-        {
-            UE_LOG_ONLINE_STOREV2(Log, TEXT("ConvertReceiptToString: No receipt data found for transaction"));
-        }
-        return ReceiptData;
+		TSharedPtr<FString> ReceiptSharedData;
+		FString ReceiptData = ConvertReceiptToString();
+		if (!ReceiptData.IsEmpty())
+		{
+			ReceiptSharedData = MakeShared<FString>(MoveTemp(ReceiptData));
+		}
+        return ReceiptSharedData;
     }
 }
 
@@ -54,22 +65,35 @@ namespace OnlinePurchaseIOSPrivate
  */
 struct FOnlinePurchaseTransactionIOS
 {
-    explicit FOnlinePurchaseTransactionIOS(SKPaymentTransaction* Transaction);
+    explicit FOnlinePurchaseTransactionIOS(EPurchaseTransactionState State, SKPaymentTransaction* Transaction, TSharedPtr<FString> NativeReceipt);
     
     /** @return a string that prints useful debug information about this transaction */
-    FString ToDebugString() const
+    FString ToDebugString(bool VeryVerbose) const
     {
-        return FString::Printf(TEXT("OfferId: %s TransactionId: %s ReceiptData: %s Error:%s"),
-                        *OfferId,
-                        *TransactionIdentifier,
-                        *ReceiptData,
-                        *ErrorStr);
+		if(VeryVerbose)
+		{
+			return FString::Printf(TEXT("OfferId: %s TransactionId: %s Date: %s ReceiptData: %s Error:%s"),
+								   *OfferId,
+								   *TransactionIdentifier,
+								   *FString([NSDateFormatter localizedStringFromDate: GetPaymentTransaction().transactionDate dateStyle:NSDateFormatterShortStyle timeStyle:NSDateFormatterShortStyle]),
+								   *GetReceiptData(),
+								   *ErrorStr);
+		}
+		else
+		{
+			return FString::Printf(TEXT("OfferId: %s TransactionId: %s Date: %s Error:%s"),
+								   *OfferId,
+								   *TransactionIdentifier,
+								   *FString([NSDateFormatter localizedStringFromDate: GetPaymentTransaction().transactionDate dateStyle:NSDateFormatterShortStyle timeStyle:NSDateFormatterShortStyle]),
+								   *ErrorStr);
+		}
     }
-    
+    /** @return transaction state */
+	EPurchaseTransactionState GetTransactionState() const { return TransactionState; }
     /** @return offer id for this transaction */
     const FString& GetOfferId() const { return OfferId; }
     /** @return receipt data for this transaction */
-    const FString& GetReceiptData() const { return ReceiptData; }
+    FString GetReceiptData() const { return ReceiptData ? *ReceiptData : FString(); }
     /** @return error string for this transaction, if applicable */
     const FString& GetErrorStr() const { return ErrorStr; }
     /** @return quantity of items requested of the same type */
@@ -83,21 +107,24 @@ private:
     /** iTunesConnect offer id */
     FString OfferId;
     /** Opaque store receipt data */
-    FString ReceiptData;
+    TSharedPtr<FString> ReceiptData;
     /** Error on the transaction, if applicable */
     FString ErrorStr;
     /** Unique transaction identifier */
     FString TransactionIdentifier;
     /** Quantity of items requested*/
     int Quantity;
+	/** The transalted transaction state */
+	EPurchaseTransactionState TransactionState;
     /** Platform transaction object*/
     TRetainedObjCInstance<SKPaymentTransaction*> PaymentTransaction;
 };
 
-FOnlinePurchaseTransactionIOS::FOnlinePurchaseTransactionIOS(SKPaymentTransaction* Transaction)
-    : ReceiptData(OnlinePurchaseIOSPrivate::ConvertReceiptToString())
+FOnlinePurchaseTransactionIOS::FOnlinePurchaseTransactionIOS(EPurchaseTransactionState State, SKPaymentTransaction* Transaction, TSharedPtr<FString> NativeReceipt)
+    : ReceiptData(NativeReceipt)
     , TransactionIdentifier(Transaction.transactionIdentifier)
     , Quantity(0)
+	, TransactionState(State)
     , PaymentTransaction(Transaction)
 {
     if(Transaction.error)
@@ -154,30 +181,36 @@ struct FOnlinePurchaseInProgressTransactionIOS
 /// SKPaymentTransactionObserver implementation
 -(void)paymentQueue: (SKPaymentQueue*)Queue updatedTransactions: (NSArray*)Transactions
 {
-    UE_LOG_ONLINE_STOREV2(Log, TEXT("FStoreKitPurchaseProxy::updatedTransactions"));
-    for(SKPaymentTransaction* Transaction in Transactions)
-    {
-        [self UpdatedTransaction : Transaction];
-    }
+    UE_LOG_ONLINE_PURCHASE(Log, TEXT("FStoreKitPurchaseProxy::updatedTransactions"));
+
+	TSharedPtr<TArray<FOnlinePurchaseTransactionIOS>> IosTransactions = [self TranslateTransactions: Transactions];
+
+	[FIOSAsyncTask CreateTaskWithBlock : ^ bool(void)
+	{
+		if(_PurchaseReceiver)
+		{
+			_PurchaseReceiver->OnUpdatedTransactions(IosTransactions ? *IosTransactions : TArray<FOnlinePurchaseTransactionIOS>());
+		}
+		return true;
+	}];
 }
 
 -(void)paymentQueue: (SKPaymentQueue*)Queue removedTransactions: (NSArray*)Transactions
 {
-    UE_LOG_ONLINE_STOREV2(Log, TEXT("FStoreKitPurchaseProxy::removedTransaction"));
+    UE_LOG_ONLINE_PURCHASE(Log, TEXT("FStoreKitPurchaseProxy::removedTransaction"));
 }
 
 -(void)paymentQueue: (SKPaymentQueue*)Queue restoreCompletedTransactionsFailedWithError: (NSError*)Error
 {
-    UE_LOG_ONLINE_STOREV2(Log, TEXT("FStoreKitPurchaseProxy::restoreCompletedTransactionsFailedWithError"));
+	EPurchaseTransactionState CompletionState = [FStoreKitPurchaseProxy TranslateError:Error];
+    UE_LOG_ONLINE_PURCHASE(Log, TEXT("FStoreKitPurchaseProxy::restoreCompletedTransactionsFailedWithError %d"), (int32)CompletionState);
 
-    EPurchaseTransactionState CompletionState = [FStoreKitPurchaseProxy TranslateError:Error];
-    
     // Notifications from SKPaymentQueue are received on main thread. Use a task to notify from game thread
     [FIOSAsyncTask CreateTaskWithBlock : ^ bool(void)
     {
         if(_PurchaseReceiver)
         {
-            _PurchaseReceiver->OnRestoreTransactionsComplete(CompletionState);
+            _PurchaseReceiver->OnQueryReceiptsComplete(false, nullptr);
         }
         return true;
     }];
@@ -185,17 +218,25 @@ struct FOnlinePurchaseInProgressTransactionIOS
 
 -(void)paymentQueueRestoreCompletedTransactionsFinished: (SKPaymentQueue*)Queue
 {
-    UE_LOG_ONLINE_STOREV2(Log, TEXT("FStoreKitPurchaseProxy::paymentQueueRestoreCompletedTransactionsFinished"));
+    UE_LOG_ONLINE_PURCHASE(Log, TEXT("FStoreKitPurchaseProxy::paymentQueueRestoreCompletedTransactionsFinished"));
 
-    // Notifications from SKPaymentQueue are received on main thread. Use a task to notify from game thread
-    [FIOSAsyncTask CreateTaskWithBlock : ^ bool(void)
-    {
-        if(_PurchaseReceiver)
-        {
-            _PurchaseReceiver->OnRestoreTransactionsComplete(EPurchaseTransactionState::Restored);
-        }
-        return true;
-    }];
+	[self NotifySuccessfulQueryReceiptsComplete];
+}
+
+-(void)NotifySuccessfulQueryReceiptsComplete
+{
+	NSArray<SKPaymentTransaction*>* KnownTransactions = [[SKPaymentQueue defaultQueue] transactions];
+	TSharedPtr<TArray<FOnlinePurchaseTransactionIOS>> IosTransactions = [self TranslateTransactions: KnownTransactions];
+
+	// Notifications from SKPaymentQueue are received on main thread. Use a task to notify from game thread
+	[FIOSAsyncTask CreateTaskWithBlock : ^ bool(void)
+	{
+		if(_PurchaseReceiver)
+		{
+			_PurchaseReceiver->OnQueryReceiptsComplete(true, IosTransactions);
+		}
+		return true;
+	}];
 }
 
 ////////////////////////////////////////////////////////////////////
@@ -204,21 +245,19 @@ struct FOnlinePurchaseInProgressTransactionIOS
 - (id)initWithReceiver: (FOnlinePurchaseIOS*)PurchaseReceiver
 {
     _PurchaseReceiver = PurchaseReceiver;
-    [FProxyPaymentTransactionObserver sharedInstance].ossObserver = self;
 
     // Notifications from SKPaymentQueue are received on main thread. Use dispatch_async on main queue to serialize operations
     dispatch_async(dispatch_get_main_queue(), ^
     {
-        SKPaymentQueue* queue = [SKPaymentQueue defaultQueue];
-        [self paymentQueue: queue updatedTransactions: [queue transactions]];
+		[[SKPaymentQueue defaultQueue] addTransactionObserver: self];
     });
     return self;
 }
 
 -(void)Shutdown
 {
+	[[SKPaymentQueue defaultQueue] removeTransactionObserver: self];
     _PurchaseReceiver = nullptr;
-    [FProxyPaymentTransactionObserver sharedInstance].ossObserver = nil;
 }
 
 - (bool)CanMakePayments
@@ -226,6 +265,15 @@ struct FOnlinePurchaseInProgressTransactionIOS
     return [SKPaymentQueue canMakePayments] != FALSE;
 }
 
+- (void)QueryKnownReceipts: (bool)completeQuery
+{
+	// Notifications from SKPaymentQueue are received on main thread. Use dispatch_async on main queue to serialize operations
+	dispatch_async(dispatch_get_main_queue(), ^
+	{
+		[self NotifySuccessfulQueryReceiptsComplete];
+	});
+}
+
 +(EPurchaseTransactionState)TranslateError: (NSError*)Error
 {
     EPurchaseTransactionState TranslatedError = EPurchaseTransactionState::Failed;
@@ -246,7 +294,36 @@ struct FOnlinePurchaseInProgressTransactionIOS
     return TranslatedError;
 }
 
--(void)UpdatedTransaction : (SKPaymentTransaction*)Transaction
+// Returns a TSharedPtr to avoid copying the contained TArray when captured in Obj-C blocks
+-(TSharedPtr<TArray<FOnlinePurchaseTransactionIOS>>)TranslateTransactions: (NSArray<SKPaymentTransaction*>*) Transactions
+{
+	TSharedPtr<TArray<FOnlinePurchaseTransactionIOS>> IosTransactions;
+	const int32 TransactionCount = [Transactions count];
+	if (TransactionCount > 0)
+	{
+		IosTransactions = MakeShared<TArray<FOnlinePurchaseTransactionIOS>>();
+		IosTransactions->Reserve(TransactionCount);
+		TSharedPtr<FString> ReceiptSnapshot;
+		
+		for (SKPaymentTransaction* Transaction in Transactions)
+		{
+			SKPaymentTransactionState State = [Transaction transactionState];
+			// The string receipt data may be big. Avoid creating it if won't be needed
+			if (ReceiptSnapshot == nullptr && (State == SKPaymentTransactionStatePurchased || State == SKPaymentTransactionStateRestored))
+			{
+				ReceiptSnapshot = OnlinePurchaseIOSPrivate::ConvertReceiptToSharedString();
+			}
+
+			if (TOptional<FOnlinePurchaseTransactionIOS> TransactionToAdd = [self TryTranslateTransaction: Transaction withReceiptData: ReceiptSnapshot])
+			{
+				IosTransactions->Emplace(MoveTemp(*TransactionToAdd));
+			}
+		}
+	}
+	return IosTransactions;
+}
+
+-(TOptional<FOnlinePurchaseTransactionIOS>)TryTranslateTransaction: (SKPaymentTransaction*)Transaction withReceiptData: (TSharedPtr<FString>)ReceiptData
 {
     TOptional<EPurchaseTransactionState> TranslatedState;
     SKPaymentTransaction* ActualTransaction = Transaction;
@@ -255,60 +332,54 @@ struct FOnlinePurchaseInProgressTransactionIOS
     switch ([Transaction transactionState])
     {
     case SKPaymentTransactionStatePurchased:
-        UE_LOG_ONLINE_STOREV2(Log, TEXT("FStoreKitPurchaseProxy::updatedTransaction SKPaymentTransactionStatePurchased"));
+        UE_LOG_ONLINE_PURCHASE(Log, TEXT("FStoreKitPurchaseProxy::TryTranslateTransaction SKPaymentTransactionStatePurchased"));
 
-        if (FParse::Param(FCommandLine::Get(), TEXT("disableiosredeem")))
-        {
-            UE_LOG_ONLINE_STOREV2(Log, TEXT("FStoreKitPurchaseProxy::completedTransaction (disabled)"));
-        }
-        else
-        {
-            TranslatedState = EPurchaseTransactionState::Purchased;
-        }
+		if (FParse::Param(FCommandLine::Get(), TEXT("disableiosredeem")))
+		{
+			UE_LOG_ONLINE_PURCHASE(Log, TEXT("FStoreKitPurchaseProxy::completedTransaction (disabled)"));
+		}
+		else
+		{
+			TranslatedState = EPurchaseTransactionState::Purchased;
+		}
         break;
     case SKPaymentTransactionStateFailed:
-        UE_LOG_ONLINE_STOREV2(Log, TEXT("FStoreKitPurchaseProxy::updatedTransaction SKPaymentTransactionStateFailed"));
+        UE_LOG_ONLINE_PURCHASE(Log, TEXT("FStoreKitPurchaseProxy::TryTranslateTransaction SKPaymentTransactionStateFailed"));
         TranslatedState = [FStoreKitPurchaseProxy TranslateError:ActualTransaction.error];
         FinishTransaction = ActualTransaction;
         break;
     case SKPaymentTransactionStateRestored:
-        UE_LOG_ONLINE_STOREV2(Log, TEXT("FStoreKitPurchaseProxy::updatedTransaction SKPaymentTransactionStateRestored"));
+        UE_LOG_ONLINE_PURCHASE(Log, TEXT("FStoreKitPurchaseProxy::TryTranslateTransaction SKPaymentTransactionStateRestored"));
         TranslatedState = EPurchaseTransactionState::Restored;
         FinishTransaction = ActualTransaction;
         ActualTransaction = ActualTransaction.originalTransaction;
         break;
     case SKPaymentTransactionStatePurchasing:
-        UE_LOG_ONLINE_STOREV2(Log, TEXT("FStoreKitPurchaseProxy::updatedTransaction SKPaymentTransactionStatePurchasing"));
+        UE_LOG_ONLINE_PURCHASE(Log, TEXT("FStoreKitPurchaseProxy::TryTranslateTransaction SKPaymentTransactionStatePurchasing"));
         break;
     case SKPaymentTransactionStateDeferred:
-        UE_LOG_ONLINE_STOREV2(Log, TEXT("FStoreKitPurchaseProxy::updatedTransaction SKPaymentTransactionStateDeferred"));
+        UE_LOG_ONLINE_PURCHASE(Log, TEXT("FStoreKitPurchaseProxy::TryTranslateTransaction SKPaymentTransactionStateDeferred"));
         TranslatedState = EPurchaseTransactionState::Deferred;
-        FinishTransaction = ActualTransaction;
         break;
     default:
-        UE_LOG_ONLINE_STOREV2(Warning, TEXT("FStoreKitPurchaseProxy::updatedTransaction unhandled state: %d"), [Transaction transactionState]);
+        UE_LOG_ONLINE_PURCHASE(Warning, TEXT("FStoreKitPurchaseProxy::TryTranslateTransaction unknown transaction state: %d"), [Transaction transactionState]);
         break;
     }
 
-    if(TranslatedState)
-    {
-        FOnlinePurchaseTransactionIOS TransactionData(ActualTransaction);
-        // Notifications from SKPaymentQueue are received on main thread. Use a task to notify from game thread
-        [FIOSAsyncTask CreateTaskWithBlock : ^ bool(void)
-        {
-            if(_PurchaseReceiver)
-            {
-                _PurchaseReceiver->OnTransactionComplete(*TranslatedState, TransactionData);
-            }
-            return true;
-        }];
-    }
-    
+	TOptional<FOnlinePurchaseTransactionIOS> TranslatedTransaction;
+	if (TranslatedState)
+	{
+		TranslatedTransaction.Emplace(*TranslatedState, ActualTransaction, MoveTemp(ReceiptData));
+	}
+	
+	// Finish transactions that should be unconditionally finish. Only leave unfinished those related to an in progress purchase so
+	// licensees can finish after validation
     if(FinishTransaction)
     {
-        UE_LOG_ONLINE_STOREV2(Log, TEXT("FStoreKitPurchaseProxy::updatedTransaction transaction finished"));
-        [[SKPaymentQueue defaultQueue] finishTransaction:FinishTransaction];
+        UE_LOG_ONLINE_PURCHASE(Log, TEXT("FStoreKitPurchaseProxy::TryTranslateTransaction transaction finished"));
+        [self FinalizeTransaction:FinishTransaction];
     }
+	return TranslatedTransaction;
 }
 
 -(void)RestorePurchases
@@ -322,7 +393,7 @@ struct FOnlinePurchaseInProgressTransactionIOS
 
 -(void)MakePurchase: (SKProduct*)Product withQuantity: (int)Quantity simulateAskToBuy: (bool) bAskToBuy
 {
-    UE_LOG_ONLINE_STOREV2(Log, TEXT("FStoreKitPurchaseProxy::makePurchase by SKProduct"));
+    UE_LOG_ONLINE_PURCHASE(Log, TEXT("FStoreKitPurchaseProxy::makePurchase by SKProduct"));
 
     // Notifications from SKPaymentQueue are received on main thread. Use dispatch_async on main queue to serialize operations
     dispatch_async(dispatch_get_main_queue(), ^
@@ -338,7 +409,7 @@ struct FOnlinePurchaseInProgressTransactionIOS
 
 -(void)FinalizeTransaction: (SKPaymentTransaction*)Transaction
 {
-    UE_LOG_ONLINE_STOREV2(Log, TEXT("FStoreKitPurchaseProxy::finalizeTransaction"));
+    UE_LOG_ONLINE_PURCHASE(Log, TEXT("FStoreKitPurchaseProxy::finalizeTransaction"));
 
     // Notifications from SKPaymentQueue are received on main thread. Use dispatch_async on main queue to serialize operations
     dispatch_async(dispatch_get_main_queue(), ^
@@ -350,7 +421,7 @@ struct FOnlinePurchaseInProgressTransactionIOS
 @end
 
 FOnlinePurchaseIOS::FOnlinePurchaseIOS(FOnlineSubsystemIOS* InSubsystem)
-	: bRestoringTransactions(false)
+	: bQueryingReceipts(false)
 	, Subsystem(InSubsystem)
 {
 	UE_LOG_ONLINE_PURCHASE(Log, TEXT( "FOnlinePurchaseIOS::FOnlinePurchaseIOS" ));
@@ -368,8 +439,8 @@ FOnlinePurchaseIOS::~FOnlinePurchaseIOS()
 void FOnlinePurchaseIOS::DumpAppReceipt()
 {
     FString receiptData = OnlinePurchaseIOSPrivate::ConvertReceiptToString();
-    UE_LOG_ONLINE_STOREV2(Verbose, TEXT("FOnlinePurchaseIOS::DumpAppReceipt"));
-    UE_LOG_ONLINE_STOREV2(Verbose, TEXT("%s"), *receiptData);
+    UE_LOG_ONLINE_PURCHASE(Verbose, TEXT("FOnlinePurchaseIOS::DumpAppReceipt"));
+	UE_LOG_ONLINE_PURCHASE(Verbose, TEXT("%s"), *receiptData);
 }
 
 bool FOnlinePurchaseIOS::IsAllowedToPurchase(const FUniqueNetId& /*UserId*/)
@@ -409,10 +480,9 @@ void FOnlinePurchaseIOS::Checkout(const FUniqueNetId& UserId, const FPurchaseChe
 #if !UE_BUILD_SHIPPING
                         bAskToBuy = OnlinePurchaseIOSPrivate::CVarSimluateAskToBuy.GetValueOnGameThread() == 1;
 #endif
-                        
+						InProgressTransaction = MakeShared<const FOnlinePurchaseInProgressTransactionIOS>(CheckoutRequest, Delegate);
                         [StoreKitProxy MakePurchase:Product withQuantity:Offer.Quantity simulateAskToBuy:bAskToBuy];
-                        InProgressTransaction = MakeShared<const FOnlinePurchaseInProgressTransactionIOS>(CheckoutRequest, Delegate);
-                        bStarted = true;
+						bStarted = true;
                     }
                 }
                 if(!bStarted)
@@ -427,7 +497,7 @@ void FOnlinePurchaseIOS::Checkout(const FUniqueNetId& UserId, const FPurchaseChe
         }
 		else
 		{
-			ErrorMessage = NSLOCTEXT("IOSPurchase", "ErrorTransactionInProgress", "Failed to checkout, user has in progress transaction.");
+			ErrorMessage = NSLOCTEXT("IOSPurchase", "ErrorTransactionInProgress", "Failed to checkout, user has in progress transaction or is retoring transactions.");
 		}
 	}
 	else
@@ -458,18 +528,19 @@ void FOnlinePurchaseIOS::FinalizePurchase(const FUniqueNetId& UserId, const FStr
 {
 	UE_LOG_ONLINE_PURCHASE(Log, TEXT("FOnlinePurchaseIOS::FinalizePurchase %s %s"), *UserId.ToString(), *ReceiptId);
 
-    int IndexToRemove = Transactions.IndexOfByPredicate([&ReceiptId](const FKnownTransaction& Entry) { return Entry.Receipt->TransactionId == ReceiptId; });
-    if(IndexToRemove != INDEX_NONE)
+    const FKnownTransaction* TransactionToFinalize = CachedReceipts.FindByPredicate([&ReceiptId](const FKnownTransaction& Entry) { return Entry.Receipt->TransactionId == ReceiptId; });
+    if(TransactionToFinalize)
     {
-        const FKnownTransaction& Transaction = Transactions[IndexToRemove];
-        
         // Restored transactions are already marked finalized when received but kept as known transactions so user code can process them
-        if(Transaction.Receipt->TransactionState != EPurchaseTransactionState::Restored)
+        if(TransactionToFinalize->Receipt->TransactionState != EPurchaseTransactionState::Restored)
         {
-            [StoreKitProxy FinalizeTransaction:Transaction.PaymentTransaction];
+            [StoreKitProxy FinalizeTransaction: TransactionToFinalize->PaymentTransaction];
         }
-        Transactions.RemoveAt(IndexToRemove);
     }
+	else
+	{
+		UE_LOG_ONLINE_PURCHASE(Warning, TEXT("FOnlinePurchaseIOS::FinalizePurchase transaction not found in CachedReceipts"));
+	}
 }
 
 void FOnlinePurchaseIOS::RedeemCode(const FUniqueNetId& UserId, const FRedeemCodeRequest& RedeemCodeRequest, const FOnPurchaseRedeemCodeComplete& Delegate)
@@ -480,29 +551,32 @@ void FOnlinePurchaseIOS::RedeemCode(const FUniqueNetId& UserId, const FRedeemCod
 
 void FOnlinePurchaseIOS::QueryReceipts(const FUniqueNetId& UserId, bool bRestoreReceipts, const FOnQueryReceiptsComplete& Delegate)
 {
-	bool bSuccess = true;
-	bool bTriggerDelegate = true;
-	if (bRestoreReceipts)
+	bool bStarted = true;
+
+	if (!bQueryingReceipts)
 	{
-		if (!bRestoringTransactions)
+		QueryReceiptsComplete = Delegate;
+		bQueryingReceipts = true;
+		if (bRestoreReceipts)
 		{
-            QueryReceiptsComplete = Delegate;
-            bTriggerDelegate = false;
-            bRestoringTransactions = true;
             [StoreKitProxy RestorePurchases];
 		}
 		else
 		{
-			UE_LOG_ONLINE_PURCHASE(Log, TEXT("FOnlinePurchaseIOS::QueryReceipts already restoring transactions"));
-			bSuccess = false;
+			[StoreKitProxy QueryKnownReceipts:true];
 		}
 	}
-	
-	if (bTriggerDelegate)
+	else
+	{
+		UE_LOG_ONLINE_PURCHASE(Log, TEXT("FOnlinePurchaseIOS::QueryReceipts already restoring transactions or checkout in progress"));
+		bStarted = false;
+	}
+
+	if (!bStarted)
 	{
 		// Query receipts comes dynamically from the StoreKit observer
-		Subsystem->ExecuteNextTick([Delegate, bSuccess]() {
-			FOnlineError Result(bSuccess);
+		Subsystem->ExecuteNextTick([Delegate]() {
+			FOnlineError Result(false);
 			Delegate.ExecuteIfBound(Result);
 		});
 	}
@@ -512,121 +586,142 @@ void FOnlinePurchaseIOS::GetReceipts(const FUniqueNetId& UserId, TArray<FPurchas
 {
 	OutReceipts.Empty();
 
-	for (const FKnownTransaction& Transaction: Transactions)
+	for (const FKnownTransaction& Transaction: CachedReceipts)
 	{
 		OutReceipts.Add(*Transaction.Receipt);
 	}
 }
 
-void FOnlinePurchaseIOS::OnTransactionComplete(EPurchaseTransactionState Result, const FOnlinePurchaseTransactionIOS& TransactionData)
+void FOnlinePurchaseIOS::OnUpdatedTransactions(const TArray<FOnlinePurchaseTransactionIOS>& Transactions)
 {
-	UE_LOG_ONLINE_PURCHASE(Log, TEXT("FOnlinePurchaseIOS::OnTransactionCompleteResponse Result=%s TransactionData=[%s]"), LexToString(Result), *TransactionData.ToDebugString());
-	
-    bool bIsInProgressTransaction = InProgressTransaction != nullptr && InProgressTransaction->HasRequestForProduct(TransactionData.GetOfferId());
-    bool bIsResoredTransaction = (Result == EPurchaseTransactionState::Restored);
-    
-	if (bIsInProgressTransaction && !bIsResoredTransaction)
-	{
-        FOnlineError FinalResult;
-        const FString& ErrorStr = TransactionData.GetErrorStr();
+	UE_LOG_ONLINE_PURCHASE(Log, TEXT("FOnlinePurchaseIOS::OnUpdatedTransactions HasInprogressTransaction: %d IsQueryingReceipts: %d"), InProgressTransaction != nullptr, bQueryingReceipts);
 
-        switch (Result)
-        {
-            case EPurchaseTransactionState::Failed:
-                FinalResult.SetFromErrorCode(TEXT("com.epicgames.purchase.failure"));
-                FinalResult.ErrorMessage = !ErrorStr.IsEmpty() ? FText::FromString(ErrorStr) : LOCTEXT("IOSTransactionFailed", "TransactionFailed");
-                break;
-            case EPurchaseTransactionState::Canceled:
-                FinalResult.SetFromErrorCode(TEXT("com.epicgames.catalog_helper.user_cancelled"));
-                FinalResult.ErrorMessage = !ErrorStr.IsEmpty() ? FText::FromString(ErrorStr) : LOCTEXT("IOSTransactionCancel", "TransactionCanceled");
-                break;
-            case EPurchaseTransactionState::Purchased:
-                FinalResult.bSucceeded = true;
-                break;
-            case EPurchaseTransactionState::Deferred:
-                FinalResult.SetFromErrorCode(TEXT("com.epicgames.purchase.deferred"));
-                FinalResult.ErrorMessage = !ErrorStr.IsEmpty() ? FText::FromString(ErrorStr) : LOCTEXT("IOSTransactionDeferred", "Transaction awaiting approval.");
-                break;
-            default:
-                UE_LOG_ONLINE_PURCHASE(Warning, TEXT("Unexpected state after purchase %d"), (int)Result);
-                FinalResult.SetFromErrorCode(TEXT("com.epicgames.purchase.unexpected_state"));
-                FinalResult.ErrorMessage = !ErrorStr.IsEmpty() ? FText::FromString(ErrorStr) : LOCTEXT("UnexpectedState", "Unexpected purchase result");
-                break;
-        }
-        
-        TSharedRef<FPurchaseReceipt> Receipt = GenerateReceipt(Result, InProgressTransaction->CheckoutRequest, TransactionData);
-        
-        if(FinalResult.bSucceeded)
-        {
-            Receipt = TryStoreTransactionAndGetReceipt(Receipt, TransactionData);
-        }
-        InProgressTransaction->CheckoutCompleteDelegate.ExecuteIfBound(FinalResult, Receipt);
-        InProgressTransaction.Reset();
-	}
-	else
+	bool bExpectingTransactions = InProgressTransaction != nullptr || bQueryingReceipts; 
+	if (InProgressTransaction != nullptr)
 	{
-		// Transactions other than the in progress one are stored for later redemption. Those can be received at login or when restoring transactions
-        // We have no interest on other failure states because they won't store a redeemable receipt
-		if (Result == EPurchaseTransactionState::Restored || Result == EPurchaseTransactionState::Purchased)
+		struct FCheckoutResult
 		{
-            TSharedRef<FPurchaseReceipt> Receipt = GenerateOfflineReceipt(Result, TransactionData);
-            TryStoreTransactionAndGetReceipt(Receipt, TransactionData);
-        }
-        
-        if(Result == EPurchaseTransactionState::Purchased)
-        {
-		    UE_LOG_ONLINE_PURCHASE(Log, TEXT("Deferred transaction finished or unfinished transaction was received"));
-            TWeakPtr<FOnlinePurchaseIOS> WeakThis = AsShared();
-            Subsystem->ExecuteNextTick([WeakThis]()
+			FOnlineError Result;
+			TSharedPtr<FPurchaseReceipt> Receipt;
+		};
+		TOptional<FCheckoutResult> CheckoutResult;
+		
+		for (const FOnlinePurchaseTransactionIOS& Transaction : Transactions)
+		{
+			const EPurchaseTransactionState TransactionState = Transaction.GetTransactionState();
+			
+			if (InProgressTransaction != nullptr &&
+				TransactionState != EPurchaseTransactionState::Restored &&
+				InProgressTransaction->HasRequestForProduct(Transaction.GetOfferId()))
 			{
-				FOnlinePurchaseIOSPtr StrongThis = WeakThis.Pin();
-				if (StrongThis.IsValid())
+				FOnlineError FinalResult;
+				const FString& ErrorStr = Transaction.GetErrorStr();
+				
+				switch (TransactionState)
+				{
+					case EPurchaseTransactionState::Failed:
+						FinalResult.SetFromErrorCode(TEXT("com.epicgames.purchase.failure"));
+						FinalResult.ErrorMessage = !ErrorStr.IsEmpty() ? FText::FromString(ErrorStr) : LOCTEXT("IOSTransactionFailed", "TransactionFailed");
+						break;
+					case EPurchaseTransactionState::Canceled:
+						FinalResult.SetFromErrorCode(TEXT("com.epicgames.catalog_helper.user_cancelled"));
+						FinalResult.ErrorMessage = !ErrorStr.IsEmpty() ? FText::FromString(ErrorStr) : LOCTEXT("IOSTransactionCancel", "TransactionCanceled");
+						break;
+					case EPurchaseTransactionState::Purchased:
+						FinalResult.bSucceeded = true;
+						break;
+					case EPurchaseTransactionState::Deferred:
+						FinalResult.SetFromErrorCode(TEXT("com.epicgames.purchase.deferred"));
+						FinalResult.ErrorMessage = !ErrorStr.IsEmpty() ? FText::FromString(ErrorStr) : LOCTEXT("IOSTransactionDeferred", "Transaction awaiting approval.");
+						break;
+					default:
+						UE_LOG_ONLINE_PURCHASE(Warning, TEXT("Unexpected state after purchase %d"), (int)TransactionState);
+						FinalResult.SetFromErrorCode(TEXT("com.epicgames.purchase.unexpected_state"));
+						FinalResult.ErrorMessage = !ErrorStr.IsEmpty() ? FText::FromString(ErrorStr) : LOCTEXT("UnexpectedState", "Unexpected purchase result");
+						break;
+				}
+				
+				TSharedRef<FPurchaseReceipt> Receipt = GenerateReceipt(InProgressTransaction->CheckoutRequest, Transaction);
+				
+				if(FinalResult.bSucceeded)
 				{
-					// No user id for app store
-					StrongThis->TriggerOnUnexpectedPurchaseReceiptDelegates(*FUniqueNetIdIOS::EmptyId());
+					AddReceiptToCache(CachedReceipts, Receipt, Transaction);
 				}
-			});
+				CheckoutResult = FCheckoutResult{MoveTemp(FinalResult), MoveTemp(Receipt)};
+				break;
+			}
+		}
+
+		if (CheckoutResult)
+		{
+			InProgressTransaction->CheckoutCompleteDelegate.ExecuteIfBound(CheckoutResult->Result, CheckoutResult->Receipt.ToSharedRef());
+			InProgressTransaction.Reset();
+		}
+		else
+		{
+			bExpectingTransactions = bQueryingReceipts;
+		}
+	}
+	
+	if (!bExpectingTransactions)
+	{
+		for (const FOnlinePurchaseTransactionIOS& Transaction : Transactions)
+		{
+			if (Transaction.GetTransactionState() == EPurchaseTransactionState::Purchased)
+			{
+				UE_LOG_ONLINE_PURCHASE(Log, TEXT("Deferred transaction finished or subscription renewal transaction was received"));
+				TWeakPtr<FOnlinePurchaseIOS> WeakThis = AsShared();
+				Subsystem->ExecuteNextTick([WeakThis]() {
+					FOnlinePurchaseIOSPtr StrongThis = WeakThis.Pin();
+					if (StrongThis.IsValid())
+					{
+						// No user id for app store
+						StrongThis->TriggerOnUnexpectedPurchaseReceiptDelegates(*FUniqueNetIdIOS::EmptyId());
+					}
+				});
+				break;
+			}
 		}
-        else if(Result == EPurchaseTransactionState::Restored)
-        {
-            UE_LOG_ONLINE_PURCHASE(Log, TEXT("Restored transaction"));
-        }
 	}
 }
 
-TSharedRef<FPurchaseReceipt> FOnlinePurchaseIOS::TryStoreTransactionAndGetReceipt(const TSharedRef<FPurchaseReceipt>& Receipt, const FOnlinePurchaseTransactionIOS& Transaction)
+void FOnlinePurchaseIOS::AddReceiptToCache(TArray<FKnownTransaction>& Cache, const TSharedRef<FPurchaseReceipt>& Receipt, const FOnlinePurchaseTransactionIOS& Transaction)
 {
-    FKnownTransaction* CachedTransaction = Transactions.FindByPredicate([&Transaction](const FKnownTransaction& KnownTransaction) {
+    FKnownTransaction* CachedTransaction = Cache.FindByPredicate([&Transaction](const FKnownTransaction& KnownTransaction) {
         return KnownTransaction.Receipt->TransactionId == Transaction.GetTransactionIdentifier();
     });
     if(CachedTransaction)
     {
-        // In case the transaction already exists as Purchased don't make it Restored
-        if(CachedTransaction->Receipt->TransactionState == EPurchaseTransactionState::Purchased &&
-           Receipt->TransactionState == EPurchaseTransactionState::Restored)
-        {
-            return CachedTransaction->Receipt;
-        }
-        else
-        {
-            CachedTransaction->Receipt = Receipt;
-            return Receipt;
-        }
+		CachedTransaction->Receipt = Receipt;
     }
     else
     {
-        Transactions.Emplace(Receipt, Transaction.GetPaymentTransaction());
-        return Receipt;
+		Cache.Emplace(Receipt, Transaction.GetPaymentTransaction());
     }
 }
 
-void FOnlinePurchaseIOS::OnRestoreTransactionsComplete(EPurchaseTransactionState Result)
+void FOnlinePurchaseIOS::OnQueryReceiptsComplete(bool bSuccess, const TSharedPtr<TArray<FOnlinePurchaseTransactionIOS>>& KnownTransactions)
 {
-	UE_LOG_ONLINE_PURCHASE(Verbose, TEXT("FOnlinePurchaseIOS::OnRestoreTransactionsComplete %d"), (int32)Result);
+	TArray< FKnownTransaction > IntermediateCachedReceipts;
 	
+	if (KnownTransactions)
+	{
+		for (const FOnlinePurchaseTransactionIOS& Transaction: *KnownTransactions)
+		{
+			const EPurchaseTransactionState TransactionState = Transaction.GetTransactionState();
+			if (TransactionState == EPurchaseTransactionState::Restored || TransactionState == EPurchaseTransactionState::Purchased)
+			{
+				TSharedRef<FPurchaseReceipt> Receipt = GenerateOfflineReceipt(Transaction);
+				AddReceiptToCache(IntermediateCachedReceipts, Receipt, Transaction);
+			}
+		}
+	}
 	// Full restore is complete
-	bRestoringTransactions = false;
-	bool bSuccess = (Result == EPurchaseTransactionState::Restored);
+	bQueryingReceipts = false;
+	if (bSuccess)
+	{
+		CachedReceipts = MoveTemp(IntermediateCachedReceipts);
+	}
 	Subsystem->ExecuteNextTick([this, bSuccess]() {
 		FOnlineError FinalResult(bSuccess);
 		QueryReceiptsComplete.ExecuteIfBound(FinalResult);
@@ -654,15 +749,16 @@ TSharedRef<FPurchaseReceipt> FOnlinePurchaseIOS::GenerateFailReceipt(const FPurc
     return Receipt;
 }
 
-TSharedRef<FPurchaseReceipt> FOnlinePurchaseIOS::GenerateOfflineReceipt(EPurchaseTransactionState Result, const FOnlinePurchaseTransactionIOS& Transaction)
+TSharedRef<FPurchaseReceipt> FOnlinePurchaseIOS::GenerateOfflineReceipt(const FOnlinePurchaseTransactionIOS& Transaction)
 {
+	const EPurchaseTransactionState TransactionState = Transaction.GetTransactionState();
     TSharedRef<FPurchaseReceipt> Receipt = MakeShared<FPurchaseReceipt>();
 
-	Receipt->TransactionState = Result;
+	Receipt->TransactionState = TransactionState;
 	Receipt->TransactionId = Transaction.GetTransactionIdentifier();
 	
-	if (Result == EPurchaseTransactionState::Purchased ||
-	    Result == EPurchaseTransactionState::Restored)
+	if (TransactionState == EPurchaseTransactionState::Purchased ||
+		TransactionState == EPurchaseTransactionState::Restored)
 	{
 		FPurchaseReceipt::FReceiptOfferEntry ReceiptEntry(TEXT(""), Transaction.GetOfferId(), Transaction.GetQuantity());
 		
@@ -678,12 +774,13 @@ TSharedRef<FPurchaseReceipt> FOnlinePurchaseIOS::GenerateOfflineReceipt(EPurchas
 	return Receipt;
 }
 
-TSharedRef<FPurchaseReceipt> FOnlinePurchaseIOS::GenerateReceipt(EPurchaseTransactionState Result, const FPurchaseCheckoutRequest& CheckoutRequest, const FOnlinePurchaseTransactionIOS& Transaction)
+TSharedRef<FPurchaseReceipt> FOnlinePurchaseIOS::GenerateReceipt(const FPurchaseCheckoutRequest& CheckoutRequest, const FOnlinePurchaseTransactionIOS& Transaction)
 {
-    TSharedRef<FPurchaseReceipt> Receipt = GenerateOfflineReceipt(Result, Transaction);
+	const EPurchaseTransactionState TransactionState = Transaction.GetTransactionState();
+    TSharedRef<FPurchaseReceipt> Receipt = GenerateOfflineReceipt(Transaction);
 
-	if(Result != EPurchaseTransactionState::Purchased &&
-       Result != EPurchaseTransactionState::Restored)
+	if(TransactionState != EPurchaseTransactionState::Purchased &&
+	   TransactionState != EPurchaseTransactionState::Restored)
 	{
         // Add the requested offers to the receipt in the event of an incomplete purchase.
         for(const auto& RequestedOffer : CheckoutRequest.PurchaseOffers)
diff --git a/Engine/Plugins/Online/IOS/OnlineSubsystemIOS/Source/Public/OnlinePurchaseIOS.h b/Engine/Plugins/Online/IOS/OnlineSubsystemIOS/Source/Public/OnlinePurchaseIOS.h
--- a/Engine/Plugins/Online/IOS/OnlineSubsystemIOS/Source/Public/OnlinePurchaseIOS.h
+++ b/Engine/Plugins/Online/IOS/OnlineSubsystemIOS/Source/Public/OnlinePurchaseIOS.h
@@ -50,21 +50,34 @@ public:
 	 * Log App Receipt content
 	 */
 	void DumpAppReceipt();
-
-    /**
-     * Method used internally by FStoreKitPurchaseProxy to notify that a transaction was completed. Not meant to be called by user code
-     */
-	void OnTransactionComplete(EPurchaseTransactionState Result, const FOnlinePurchaseTransactionIOS& TransactionData);
-
+	
+	/**
+	 * Method used internally by FStoreKitPurchaseProxy to handle a known transactions. Not meant to be called by user code
+	 */
+	void OnUpdatedTransactions(const TArray<FOnlinePurchaseTransactionIOS>& Transactions);
+	
     /**
      * Method used internally by FStoreKitPurchaseProxy to notify that transactions were restored. Not meant to be called by user code
      */
-    void OnRestoreTransactionsComplete(EPurchaseTransactionState Result);
+    void OnQueryReceiptsComplete(bool bSuccess, const TSharedPtr<TArray<FOnlinePurchaseTransactionIOS>>& Transactions);
 private:
-    /**
-     * Add Transaction data to cache avoiding duplicates and return proper receipt
+	/** Entry data for a known completed transaction (restored or purchased) */
+	struct FKnownTransaction
+	{
+		explicit FKnownTransaction(TSharedRef<FPurchaseReceipt> InReceipt, SKPaymentTransaction* InTransaction)
+		: Receipt(InReceipt)
+		, PaymentTransaction(InTransaction)
+		{
+		}
+
+		TSharedRef<FPurchaseReceipt> Receipt;
+		TRetainedObjCInstance<SKPaymentTransaction*> PaymentTransaction;
+	};
+
+	/**
+     * Add receipt information for transaction data to intermediate cache avoiding duplicates and return stored receipt
      */
-    TSharedRef<FPurchaseReceipt> TryStoreTransactionAndGetReceipt(const TSharedRef<FPurchaseReceipt>& Receipt, const FOnlinePurchaseTransactionIOS& Transaction);
+	static void AddReceiptToCache(TArray<FKnownTransaction>& Cache, const TSharedRef<FPurchaseReceipt>& Receipt, const FOnlinePurchaseTransactionIOS& Transaction);
 
     /**
      * Generate a receipt for a failed purchase
@@ -74,19 +87,19 @@ private:
     /**
      * Generate a receipt for a successful purchase which was received offline
      */
-    static TSharedRef<FPurchaseReceipt> GenerateOfflineReceipt(EPurchaseTransactionState Result, const FOnlinePurchaseTransactionIOS& Transaction);
+    static TSharedRef<FPurchaseReceipt> GenerateOfflineReceipt(const FOnlinePurchaseTransactionIOS& Transaction);
 
     /**
      * Generate a final receipt for the purchase currently in progress
      */
-    static TSharedRef<FPurchaseReceipt> GenerateReceipt(EPurchaseTransactionState Result, const FPurchaseCheckoutRequest& CheckoutRequest, const FOnlinePurchaseTransactionIOS& Transaction);
+    static TSharedRef<FPurchaseReceipt> GenerateReceipt(const FPurchaseCheckoutRequest& CheckoutRequest, const FOnlinePurchaseTransactionIOS& Transaction);
 private:
 	
 	/** Proxy helper to communicate from/to StoreKit*/
     TRetainedObjCInstance<FStoreKitPurchaseProxy*> StoreKitProxy;
 	
-	/** Are transactions current being restored */
-	bool bRestoringTransactions;
+	/** Are receipts being queried */
+	bool bQueryingReceipts;
 
 	/** Transient delegate to fire when restoring transactions (QueryReceipts with bRestoreReceipts=true)*/
 	FOnQueryReceiptsComplete QueryReceiptsComplete;
@@ -94,21 +107,9 @@ private:
 	/** Keeps track of in progress user transaction */
     TSharedPtr<const FOnlinePurchaseInProgressTransactionIOS> InProgressTransaction;
 	
-	/** Entry data for a known completed transaction (restored or purchased) */
-    struct FKnownTransaction
-    {
-        explicit FKnownTransaction(TSharedRef<FPurchaseReceipt> InReceipt, SKPaymentTransaction* InTransaction)
-        : Receipt(InReceipt)
-        , PaymentTransaction(InTransaction)
-        {
-        }
-
-        TSharedRef<FPurchaseReceipt> Receipt;
-        TRetainedObjCInstance<SKPaymentTransaction*> PaymentTransaction;
-    };
-	/** Cache of completed transactions */
-    TArray< FKnownTransaction > Transactions;
-	
+	/** Cache of known transactions */
+	TArray< FKnownTransaction > CachedReceipts;
+
 	/** Reference to the parent subsystem */
 	FOnlineSubsystemIOS* Subsystem;
 };
diff --git a/Engine/Source/Runtime/ApplicationCore/Private/IOS/IOSAppDelegate.cpp b/Engine/Source/Runtime/ApplicationCore/Private/IOS/IOSAppDelegate.cpp
--- a/Engine/Source/Runtime/ApplicationCore/Private/IOS/IOSAppDelegate.cpp
+++ b/Engine/Source/Runtime/ApplicationCore/Private/IOS/IOSAppDelegate.cpp
@@ -20,7 +20,6 @@
 #include "IOS/IOSAsyncTask.h"
 #include "Misc/ConfigCacheIni.h"
 #include "IOS/IOSPlatformCrashContext.h"
-#include "IOS/ProxyPaymentTransactionObserver.h"
 #include "Misc/OutputDeviceError.h"
 #include "Misc/OutputDeviceRedirector.h"
 #include "Misc/FeedbackContext.h"
@@ -1100,8 +1099,6 @@ static FAutoConsoleVariableRef CVarGEnableThermalsReport(
 	[[NSNotificationCenter defaultCenter] addObserver:self selector:@selector(OnVoiceOverStatusChanged) name:UIAccessibilityVoiceOverStatusDidChangeNotification object:nil];
 #endif
 
-	[[SKPaymentQueue defaultQueue] addTransactionObserver:[FProxyPaymentTransactionObserver sharedInstance]];
-    
 	return YES;
 }
 
@@ -1488,8 +1485,6 @@ extern double GCStartTime;
 	 */
 	FCoreDelegates::ApplicationWillTerminateDelegate.Broadcast();
     
-    [[SKPaymentQueue defaultQueue] removeTransactionObserver:[FProxyPaymentTransactionObserver sharedInstance]];
-
     // note that we are shutting down
     // TODO: fix the reason why we are hanging when asked to shutdown
 /*    RequestEngineExit(TEXT("IOS applicationWillTerminate"));
diff --git a/Engine/Source/Runtime/ApplicationCore/Private/IOS/ProxyPaymentTransactionObserver.cpp b/Engine/Source/Runtime/ApplicationCore/Private/IOS/ProxyPaymentTransactionObserver.cpp
deleted file mode 100644
diff --git a/Engine/Source/Runtime/ApplicationCore/Public/IOS/ProxyPaymentTransactionObserver.h b/Engine/Source/Runtime/ApplicationCore/Public/IOS/ProxyPaymentTransactionObserver.h
deleted file mode 100644
