Fix crash in FOnlineSessionIOS::StartSession when creating sessions. An exception was thrown because string argument is not allowed to be empty

#changelist validated
#virtualized

diff --git a/Engine/Plugins/Online/IOS/OnlineSubsystemIOS/Source/Private/OnlineSessionInterfaceIOS.cpp b/Engine/Plugins/Online/IOS/OnlineSubsystemIOS/Source/Private/OnlineSessionInterfaceIOS.cpp
--- a/Engine/Plugins/Online/IOS/OnlineSubsystemIOS/Source/Private/OnlineSessionInterfaceIOS.cpp
+++ b/Engine/Plugins/Online/IOS/OnlineSubsystemIOS/Source/Private/OnlineSessionInterfaceIOS.cpp
@@ -10,11 +10,38 @@
 @synthesize Session;
 @synthesize PeerID;
 
-- (instancetype)initSessionWithName:(NSString*) sessionName
+- (NSString *)fixDisplayName:(NSString *)displayName
+{
+	// initWithDisplayName: expects the display name to have a maximum length of 63 bytes in UTF-8 encoding, not be nil nor empty
+	static constexpr NSUInteger MaxValidDisplayNameSize = 63;
+
+	if (displayName != nil)
+	{
+		if ([displayName lengthOfBytesUsingEncoding: NSUTF8StringEncoding] > MaxValidDisplayNameSize)
+		{
+			char Buffer[MaxValidDisplayNameSize] = {0};
+			NSUInteger UsedLength = 0;
+			NSRange Range = NSMakeRange(0, [displayName length]);
+			[displayName getBytes: Buffer maxLength: MaxValidDisplayNameSize usedLength: &UsedLength encoding: NSUTF8StringEncoding options: 0 range: Range remainingRange: NULL];
+			if (NSString *FixedDisplayName = [[NSString alloc] initWithBytes: Buffer length: UsedLength encoding: NSUTF8StringEncoding])
+			{
+				return FixedDisplayName;
+			}
+		}
+		else if (displayName.length > 0)
+		{
+			return displayName;
+		}
+	}
+	return @"Unknown";
+}
+
+- (instancetype)initSessionWithName:(NSString*) sessionName andDisplayName:(NSString*) displayName
 {
     UE_LOG_ONLINE_SESSION(Display, TEXT("- (void)initSessionWithName:(NSString*) sessionName"));
 	self = [super init];
-    self.PeerID = [[[MCPeerID alloc] initWithDisplayName:@""] autorelease];
+	NSString* FixedDisplayName = [self fixDisplayName: displayName];
+	self.PeerID = [[[MCPeerID alloc] initWithDisplayName: FixedDisplayName] autorelease];
     self.Session = [[[MCSession alloc] initWithPeer:self.PeerID] autorelease];
 	self.Session.delegate = self;
 	return self;
@@ -86,14 +113,14 @@
 @implementation FGameCenterSessionDelegate
 @synthesize SessionMC;
 
-- (instancetype)initSessionWithName:(NSString*) sessionName
+- (instancetype)initSessionWithName:(NSString*) sessionName andDisplayName: (NSString*)displayName
 {
-    UE_LOG_ONLINE_SESSION(Display, TEXT("- (void)initSessionWithName:(NSString*) sessionName"));
+    UE_LOG_ONLINE_SESSION(Display, TEXT("- (void)initSessionWithName:(NSString*) sessionName (NSString*)displayName"));
 	self = [super init];
 	// Create the session object
     if ([MCSession class])
     {
-        self.SessionMC = [[[FGameCenterSessionDelegateMC alloc] initSessionWithName:sessionName] autorelease];
+        self.SessionMC = [[[FGameCenterSessionDelegateMC alloc] initSessionWithName:sessionName andDisplayName: displayName] autorelease];
     }
     
     return self;
@@ -309,7 +336,11 @@ bool FOnlineSessionIOS::StartSession(FName SessionName)
 		// Find the linked GK session and start it.
         FGameCenterSessionDelegate* LinkedGKSession = *GKSessions.Find( SessionName );
 		NSString* SafeSessionName = [NSString stringWithFString:SessionName.ToString()];
-		[LinkedGKSession initSessionWithName:SafeSessionName];
+
+		IOnlineIdentityPtr IdentityInterface = IOSSubsystem->GetIdentityInterface();
+		NSString* Nickname = IdentityInterface->GetPlayerNickname(0).GetNSString();
+
+		[LinkedGKSession initSessionWithName:SafeSessionName andDisplayName: Nickname];
 
 		// Update the session state as we are now running.
 		Session->SessionState = EOnlineSessionState::InProgress;
diff --git a/Engine/Plugins/Online/IOS/OnlineSubsystemIOS/Source/Private/OnlineSessionInterfaceIOS.h b/Engine/Plugins/Online/IOS/OnlineSubsystemIOS/Source/Private/OnlineSessionInterfaceIOS.h
--- a/Engine/Plugins/Online/IOS/OnlineSubsystemIOS/Source/Private/OnlineSessionInterfaceIOS.h
+++ b/Engine/Plugins/Online/IOS/OnlineSubsystemIOS/Source/Private/OnlineSessionInterfaceIOS.h
@@ -19,10 +19,10 @@
 @property (nonatomic, strong) MCPeerID *PeerID;
 @property (nonatomic, strong) MCSession *Session;
 
--(instancetype) initSessionWithName:(NSString*) sessionName;
+-(instancetype) initSessionWithName:(NSString*) sessionName andDisplayName:(NSString*) displayName;
 -(void) shutdownSession;
 -(bool) sessionsAvailable;
--(void)joinSession;
+-(void) joinSession;
 
 @end
 
@@ -32,10 +32,10 @@
 
 @property (nonatomic, strong) FGameCenterSessionDelegateMC *SessionMC;
 
--(instancetype) initSessionWithName:(NSString*) sessionName;
+-(instancetype) initSessionWithName:(NSString*) sessionName andDisplayName:(NSString*) displayName;
 -(void) shutdownSession;
 -(bool) sessionsAvailable;
--(void)joinSession;
+-(void) joinSession;
 
 @end
 
