Remove storekit framework for apps not distributed through the AppStore

diff --git a/Engine/Config/BaseEngine.ini b/Engine/Config/BaseEngine.ini
--- a/Engine/Config/BaseEngine.ini
+++ b/Engine/Config/BaseEngine.ini
@@ -3033,6 +3033,7 @@ DebugInfoStartY=60.0
 ; These are the defaults for iOS settings, and they need to be in the .ini since UBT reads the .ini settings, without instantiating the class
 ; Note: These should be in IOSEngine, but are placed here because the Settings editor UI cannot write to the Game IOSEngine, meaning settings in Engine/IOSEngine are non-overridable
 [/Script/IOSRuntimeSettings.IOSRuntimeSettings]
+bEnableStoreKitSupport=True
 bEnableGameCenterSupport=False
 bSupportsPortraitOrientation=False
 bSupportsITunesFileSharing=False
diff --git a/Engine/Plugins/Online/IOS/OnlineSubsystemIOS/Source/Private/OnlinePurchaseIOS.cpp b/Engine/Plugins/Online/IOS/OnlineSubsystemIOS/Source/Private/OnlinePurchaseIOS.cpp
--- a/Engine/Plugins/Online/IOS/OnlineSubsystemIOS/Source/Private/OnlinePurchaseIOS.cpp
+++ b/Engine/Plugins/Online/IOS/OnlineSubsystemIOS/Source/Private/OnlinePurchaseIOS.cpp
@@ -1,5 +1,7 @@
 // Copyright Epic Games, Inc. All Rights Reserved.
 
+#if UE_WITH_STORE_KIT
+
 #include "OnlinePurchaseIOS.h"
 #include "OnlineError.h"
 #include "OnlineSubsystemIOS.h"
@@ -794,3 +796,4 @@ TSharedRef<FPurchaseReceipt> FOnlinePurchaseIOS::GenerateReceipt(const FPurchase
 
 #undef LOCTEXT_NAMESPACE
 
+#endif //UE_WITH_STORE_KIT
diff --git a/Engine/Plugins/Online/IOS/OnlineSubsystemIOS/Source/Private/OnlineStoreIOS.cpp b/Engine/Plugins/Online/IOS/OnlineSubsystemIOS/Source/Private/OnlineStoreIOS.cpp
--- a/Engine/Plugins/Online/IOS/OnlineSubsystemIOS/Source/Private/OnlineStoreIOS.cpp
+++ b/Engine/Plugins/Online/IOS/OnlineSubsystemIOS/Source/Private/OnlineStoreIOS.cpp
@@ -1,5 +1,7 @@
 // Copyright Epic Games, Inc. All Rights Reserved.
 
+#if UE_WITH_STORE_KIT
+
 #include "OnlineStoreIOS.h"
 #include "Internationalization/Internationalization.h"
 #include "Internationalization/Culture.h"
@@ -312,3 +314,5 @@ SKProduct* FOnlineStoreIOS::GetSKProductByOfferId(const FUniqueOfferId& OfferId)
 	
 	return nil;
 }
+
+#endif //UE_WITH_STORE_KIT
diff --git a/Engine/Plugins/Online/IOS/OnlineSubsystemIOS/Source/Private/OnlineSubsystemIOS.cpp b/Engine/Plugins/Online/IOS/OnlineSubsystemIOS/Source/Private/OnlineSubsystemIOS.cpp
--- a/Engine/Plugins/Online/IOS/OnlineSubsystemIOS/Source/Private/OnlineSubsystemIOS.cpp
+++ b/Engine/Plugins/Online/IOS/OnlineSubsystemIOS/Source/Private/OnlineSubsystemIOS.cpp
@@ -164,12 +164,13 @@ bool FOnlineSubsystemIOS::Init()
 		UserCloudInterface = MakeShareable(new FOnlineUserCloudInterfaceIOS());
 		SharedCloudInterface = MakeShareable(new FOnlineSharedCloudInterfaceIOS());
 
+#if UE_WITH_STORE_KIT
 		if (IsInAppPurchasingEnabled())
 		{
 			StoreV2Interface = MakeShareable(new FOnlineStoreIOS(this));
 			PurchaseInterface = MakeShareable(new FOnlinePurchaseIOS(this));
 		}
-
+#endif
 		if (UserCloudInterface && IsCloudKitEnabled())
 		{
 			FString IOSCloudKitSyncStrategy = "";
@@ -328,8 +329,11 @@ bool FOnlineSubsystemIOS::IsInAppPurchasingEnabled()
 	GConfig->GetBool(TEXT("OnlineSubsystemIOS.Store"), TEXT("bSupportsInAppPurchasing"), bEnableIAP, GEngineIni);
 	
 	bool bEnableIAP1 = false;
-	GConfig->GetBool(TEXT("OnlineSubsystemIOS.Store"), TEXT("bSupportInAppPurchasing"), bEnableIAP1, GEngineIni);
-	return bEnableIAP || bEnableIAP1;
+	if (GConfig->GetBool(TEXT("OnlineSubsystemIOS.Store"), TEXT("bSupportInAppPurchasing"), bEnableIAP1, GEngineIni))
+	{
+		UE_LOG_ONLINE(Warning, TEXT("Deprecated use of bSupportInAppPurchasing setting in OnlineSubsystemIOS.Store. Use bSupportsInAppPurchasing instead"));
+	}
+	return UE_WITH_STORE_KIT && (bEnableIAP || bEnableIAP1);
 }
 
 NSString* FOnlineSubsystemIOS::GetPlayerId(GKPlayer* Player)
diff --git a/Engine/Plugins/Online/IOS/OnlineSubsystemIOS/Source/Public/OnlinePurchaseIOS.h b/Engine/Plugins/Online/IOS/OnlineSubsystemIOS/Source/Public/OnlinePurchaseIOS.h
--- a/Engine/Plugins/Online/IOS/OnlineSubsystemIOS/Source/Public/OnlinePurchaseIOS.h
+++ b/Engine/Plugins/Online/IOS/OnlineSubsystemIOS/Source/Public/OnlinePurchaseIOS.h
@@ -2,6 +2,8 @@
 
 #pragma once
 
+#if UE_WITH_STORE_KIT
+
 #include "Interfaces/OnlinePurchaseInterface.h"
 #include "RetainedObjCInstance.h"
 
@@ -115,3 +117,5 @@ private:
 };
 
 typedef TSharedPtr<FOnlinePurchaseIOS, ESPMode::ThreadSafe> FOnlinePurchaseIOSPtr;
+
+#endif //UE_WITH_STORE_KIT
diff --git a/Engine/Plugins/Online/IOS/OnlineSubsystemIOS/Source/Public/OnlineStoreIOS.h b/Engine/Plugins/Online/IOS/OnlineSubsystemIOS/Source/Public/OnlineStoreIOS.h
--- a/Engine/Plugins/Online/IOS/OnlineSubsystemIOS/Source/Public/OnlineStoreIOS.h
+++ b/Engine/Plugins/Online/IOS/OnlineSubsystemIOS/Source/Public/OnlineStoreIOS.h
@@ -2,6 +2,8 @@
 
 #pragma once
 
+#if UE_WITH_STORE_KIT
+
 #include "Interfaces/OnlineStoreInterfaceV2.h"
 #include "RetainedObjCInstance.h"
 
@@ -108,3 +110,4 @@ private:
 };
 
 typedef TSharedPtr<FOnlineStoreIOS, ESPMode::ThreadSafe> FOnlineStoreIOSPtr;
+#endif // UE_WITH_STORE_KIT
diff --git a/Engine/Plugins/Runtime/Firebase/Source/Firebase.Build.cs b/Engine/Plugins/Runtime/Firebase/Source/Firebase.Build.cs
--- a/Engine/Plugins/Runtime/Firebase/Source/Firebase.Build.cs
+++ b/Engine/Plugins/Runtime/Firebase/Source/Firebase.Build.cs
@@ -90,7 +90,9 @@ namespace UnrealBuildTool.Rules
 					""
 					)
 				);
-				
+				// GoogleAppMeasurement depends on Storekit
+				PublicFrameworks.Add("StoreKit");
+
 				PublicAdditionalFrameworks.Add(
 					new Framework(
 					"GoogleUtilities",
diff --git a/Engine/Source/Programs/UnrealBuildTool/Platform/IOS/UEBuildIOS.cs b/Engine/Source/Programs/UnrealBuildTool/Platform/IOS/UEBuildIOS.cs
--- a/Engine/Source/Programs/UnrealBuildTool/Platform/IOS/UEBuildIOS.cs
+++ b/Engine/Source/Programs/UnrealBuildTool/Platform/IOS/UEBuildIOS.cs
@@ -288,6 +288,12 @@ namespace UnrealBuildTool
 		[ConfigFile(ConfigHierarchyType.Engine, "/Script/IOSRuntimeSettings.IOSRuntimeSettings", "SigningCertificate")]
 		public readonly string SigningCertificate = "";
 
+		/// <summary>
+		/// true if notifications are enabled
+		/// </summary>
+		[ConfigFile(ConfigHierarchyType.Engine, "/Script/IOSRuntimeSettings.IOSRuntimeSettings", "bEnableStoreKitSupport")]
+		public readonly bool bEnableStoreKitSupport = true;
+
 		/// <summary>
 		/// true if notifications are enabled
 		/// </summary>
@@ -1212,9 +1218,17 @@ namespace UnrealBuildTool
 			CompileEnvironment.Definitions.Add("MINIMUM_UE_COMPILED_IOS_VERSION=" + TargetNum);
 
 			LinkEnvironment.AdditionalFrameworks.Add(new UEBuildFramework("GameKit"));
-			LinkEnvironment.AdditionalFrameworks.Add(new UEBuildFramework("StoreKit"));
 			LinkEnvironment.AdditionalFrameworks.Add(new UEBuildFramework("DeviceCheck"));
 
+			if (ProjectSettings.bEnableStoreKitSupport)
+			{
+				CompileEnvironment.Definitions.Add("UE_WITH_STORE_KIT=1");
+				LinkEnvironment.AdditionalFrameworks.Add(new UEBuildFramework("StoreKit"));
+			}
+			else
+			{
+				CompileEnvironment.Definitions.Add("UE_WITH_STORE_KIT=0");
+			}
 		}
 
 		/// <summary>
diff --git a/Engine/Source/Runtime/Core/Core.Build.cs b/Engine/Source/Runtime/Core/Core.Build.cs
--- a/Engine/Source/Runtime/Core/Core.Build.cs
+++ b/Engine/Source/Runtime/Core/Core.Build.cs
@@ -120,7 +120,8 @@ public class Core : ModuleRules
 			AddEngineThirdPartyPrivateStaticDependencies(Target,
 				"zlib"
 				);
-			PublicFrameworks.AddRange(new string[] { "UIKit", "Foundation", "AudioToolbox", "AVFoundation", "GameKit", "StoreKit", "CoreVideo", "CoreMedia", "CoreGraphics", "GameController", "SystemConfiguration", "DeviceCheck", "UserNotifications" });
+			PublicFrameworks.AddRange(new string[] { "UIKit", "Foundation", "AudioToolbox", "AVFoundation", "GameKit", "CoreVideo", "CoreMedia", "CoreGraphics", "GameController", "SystemConfiguration", "DeviceCheck", "UserNotifications" });
+
 			if (Target.Platform == UnrealTargetPlatform.IOS)
 			{
 				PublicFrameworks.AddRange(new string[] { "CoreMotion", "AdSupport", "WebKit" });
diff --git a/Engine/Source/Runtime/Core/Private/IOS/IOSPlatformMisc.cpp b/Engine/Source/Runtime/Core/Private/IOS/IOSPlatformMisc.cpp
--- a/Engine/Source/Runtime/Core/Private/IOS/IOSPlatformMisc.cpp
+++ b/Engine/Source/Runtime/Core/Private/IOS/IOSPlatformMisc.cpp
@@ -45,7 +45,9 @@
 #import <mach-o/dyld.h>
 #include <netinet/in.h>
 #include <SystemConfiguration/SystemConfiguration.h>
+#if UE_WITH_STORE_KIT
 #import <StoreKit/StoreKit.h>
+#endif
 #if !PLATFORM_VISIONOS
 #import <UserNotifications/UserNotifications.h>
 #endif
@@ -920,7 +922,7 @@ bool FIOSPlatformMisc::GetDiskTotalAndFreeSpace(const FString& InPath, uint64& T
 
 void FIOSPlatformMisc::RequestStoreReview()
 {
-#if !PLATFORM_TVOS
+#if UE_WITH_STORE_KIT && !PLATFORM_TVOS
     [SKStoreReviewController requestReviewInScene:[[[[UIApplication sharedApplication] delegate] window] windowScene]];
 #endif
 }
diff --git a/Engine/Source/Runtime/IOS/IOSRuntimeSettings/Classes/IOSRuntimeSettings.h b/Engine/Source/Runtime/IOS/IOSRuntimeSettings/Classes/IOSRuntimeSettings.h
--- a/Engine/Source/Runtime/IOS/IOSRuntimeSettings/Classes/IOSRuntimeSettings.h
+++ b/Engine/Source/Runtime/IOS/IOSRuntimeSettings/Classes/IOSRuntimeSettings.h
@@ -179,6 +179,10 @@ class IOSRUNTIMESETTINGS_API UIOSRuntimeSettings : public UObject
 public:
 	GENERATED_UCLASS_BODY()
 
+	// Should StoreKit support be enabled?
+	UPROPERTY(GlobalConfig, EditAnywhere, Category = Online, meta = (ConfigHierarchyEditable))
+	bool bEnableStoreKitSupport;
+
 	// Should Game Center support (iOS Online Subsystem) be enabled?
 	UPROPERTY(GlobalConfig, EditAnywhere, Category = Online, meta = (ConfigHierarchyEditable))
     bool bEnableGameCenterSupport;
diff --git a/Engine/Source/Runtime/IOS/IOSRuntimeSettings/Private/IOSRuntimeSettings.cpp b/Engine/Source/Runtime/IOS/IOSRuntimeSettings/Private/IOSRuntimeSettings.cpp
--- a/Engine/Source/Runtime/IOS/IOSRuntimeSettings/Private/IOSRuntimeSettings.cpp
+++ b/Engine/Source/Runtime/IOS/IOSRuntimeSettings/Private/IOSRuntimeSettings.cpp
@@ -25,6 +25,7 @@ UIOSRuntimeSettings::UIOSRuntimeSettings(const FObjectInitializer& ObjectInitial
 	, MinSampleRate(8000)
 	, CompressionQualityModifier(1)
 {
+	bEnableStoreKitSupport = true;
 	bEnableGameCenterSupport = true;
 	bEnableCloudKitSupport = false;
 	bUserSwitching = false;
