Support for stats deprecation extended to iOS platform

diff --git a/Engine/Plugins/Online/IOS/OnlineSubsystemIOS/Source/Private/OnlineAchievementsInterfaceIOS.cpp b/Engine/Plugins/Online/IOS/OnlineSubsystemIOS/Source/Private/OnlineAchievementsInterfaceIOS.cpp
--- a/Engine/Plugins/Online/IOS/OnlineSubsystemIOS/Source/Private/OnlineAchievementsInterfaceIOS.cpp
+++ b/Engine/Plugins/Online/IOS/OnlineSubsystemIOS/Source/Private/OnlineAchievementsInterfaceIOS.cpp
@@ -143,7 +143,7 @@ void FOnlineAchievementsIOS::WriteAchievements(const FUniqueNetId& PlayerId, FOn
 		const FVariantData& Stat = It.Value();
 
 		// Create an achievement object which should be reported to the server.
-		const FString AchievementName(It.Key().ToString());
+		const FString AchievementName(It.Key());
 		NSString* AchievementID = [NSString stringWithFString:AchievementName];
 		
 		GKAchievement* Achievement = [[[GKAchievement alloc] initWithIdentifier:AchievementID] autorelease];
diff --git a/Engine/Plugins/Online/IOS/OnlineSubsystemIOS/Source/Private/OnlineLeaderboardsInterfaceIOS.cpp b/Engine/Plugins/Online/IOS/OnlineSubsystemIOS/Source/Private/OnlineLeaderboardsInterfaceIOS.cpp
--- a/Engine/Plugins/Online/IOS/OnlineSubsystemIOS/Source/Private/OnlineLeaderboardsInterfaceIOS.cpp
+++ b/Engine/Plugins/Online/IOS/OnlineSubsystemIOS/Source/Private/OnlineLeaderboardsInterfaceIOS.cpp
@@ -102,7 +102,7 @@ bool FOnlineLeaderboardsIOS::ReadLeaderboardCompletionDelegate(NSArray* players,
                         
                             default:
                             {
-                                UE_LOG_ONLINE_LEADERBOARD(Warning, TEXT("Unsupported key value pair during retrieval from GameCenter %s"), *ColumnMeta.ColumnName.ToString());
+                                UE_LOG_ONLINE_LEADERBOARD(Warning, TEXT("Unsupported key value pair during retrieval from GameCenter %s"), *ColumnMeta.ColumnName);
                                 break;
                             }
                         }
