MarketplaceKit module is only for iOS

diff --git a/Engine/Plugins/Online/IOS/OnlineSubsystemIOS/Source/OnlineSubsystemIOS.Build.cs b/Engine/Plugins/Online/IOS/OnlineSubsystemIOS/Source/OnlineSubsystemIOS.Build.cs
--- a/Engine/Plugins/Online/IOS/OnlineSubsystemIOS/Source/OnlineSubsystemIOS.Build.cs
+++ b/Engine/Plugins/Online/IOS/OnlineSubsystemIOS/Source/OnlineSubsystemIOS.Build.cs
@@ -26,7 +26,6 @@ public class OnlineSubsystemIOS : ModuleRules
                 "Engine", 
                 "HTTP",
                 "IOSPlatformFeatures",
-                "MarketplaceKit",
 				"OnlineSubsystem", 
                 "Sockets",
             }
@@ -36,6 +35,7 @@ public class OnlineSubsystemIOS : ModuleRules
 		if (Target.Platform == UnrealTargetPlatform.IOS)
 		{
 			PublicWeakFrameworks.Add("MultipeerConnectivity");
+			PrivateDependencyModuleNames.Add("MarketplaceKit");
 		}
 	}
 }
diff --git a/Engine/Plugins/Online/IOS/OnlineSubsystemIOS/Source/Private/OnlineSubsystemIOS.cpp b/Engine/Plugins/Online/IOS/OnlineSubsystemIOS/Source/Private/OnlineSubsystemIOS.cpp
--- a/Engine/Plugins/Online/IOS/OnlineSubsystemIOS/Source/Private/OnlineSubsystemIOS.cpp
+++ b/Engine/Plugins/Online/IOS/OnlineSubsystemIOS/Source/Private/OnlineSubsystemIOS.cpp
@@ -3,7 +3,9 @@
 #include "OnlineSubsystemIOS.h"
 
 #include "IOS/IOSAppDelegate.h"
+#if !PLATFORM_TVOS && !PLATFORM_VISIONOS
 #include "MarketplaceKit.h"
+#endif
 #include "Misc/ConfigCacheIni.h"
 #import "OnlineAppStoreUtils.h"
 #include "Stats/Stats.h"
@@ -169,7 +171,9 @@ bool FOnlineSubsystemIOS::Init()
 #if UE_WITH_STORE_KIT
 		if (IsInAppPurchasingEnabled())
 		{
+#if !PLATFORM_TVOS && !PLATFORM_VISIONOS		// tvOS & VisionOS ONLY support AppStore
 			if (FMarketplaceKitModule::GetCurrentTypeStatic() == EMarketplaceType::AppStore)
+#endif
 			{
 				StoreV2Interface = MakeShareable(new FOnlineStoreIOS(this));
 				PurchaseInterface = MakeShareable(new FOnlinePurchaseIOS(this));
