- added start of StoreKit2 support in OnlineStoreIOS
- added a test refund flow call

diff --git a/Engine/Plugins/Online/IOS/OnlineSubsystemIOS/Source/OnlineSubsystemIOS.Build.cs b/Engine/Plugins/Online/IOS/OnlineSubsystemIOS/Source/OnlineSubsystemIOS.Build.cs
--- a/Engine/Plugins/Online/IOS/OnlineSubsystemIOS/Source/OnlineSubsystemIOS.Build.cs
+++ b/Engine/Plugins/Online/IOS/OnlineSubsystemIOS/Source/OnlineSubsystemIOS.Build.cs
@@ -36,6 +36,7 @@ public class OnlineSubsystemIOS : ModuleRules
 		{
 			PublicWeakFrameworks.Add("MultipeerConnectivity");
 			PrivateDependencyModuleNames.Add("MarketplaceKit");
+			SwiftInteropHeader = System.IO.Path.Combine(ModuleDirectory, "Private/StoreKitSwiftInterop.h");
 		}
 	}
 }
diff --git a/Engine/Plugins/Online/IOS/OnlineSubsystemIOS/Source/Private/OnlineStoreIOS.cpp b/Engine/Plugins/Online/IOS/OnlineSubsystemIOS/Source/Private/OnlineStoreIOS.cpp
--- a/Engine/Plugins/Online/IOS/OnlineSubsystemIOS/Source/Private/OnlineStoreIOS.cpp
+++ b/Engine/Plugins/Online/IOS/OnlineSubsystemIOS/Source/Private/OnlineStoreIOS.cpp
@@ -11,6 +11,10 @@
 #import <StoreKit/SKProductsRequest.h>
 #include "Misc/ConfigCacheIni.h"
 
+#include "StoreKit2-Swift.h"
+#include "StoreKitSwiftInterop.h"
+#include "IOSAppDelegate.h"
+
 @interface FSKProductsRequestHelper : SKProductsRequest
 {
 };
@@ -315,4 +319,33 @@ SKProduct* FOnlineStoreIOS::GetSKProductByOfferId(const FUniqueOfferId& OfferId)
 	return nil;
 }
 
+#if !UE_BUILD_SHIPPING
+class FStoreExec : public FSelfRegisteringExec
+{
+	virtual bool Exec_Dev(class UWorld* InWorld, const TCHAR* Cmd, FOutputDevice& Ar) override
+	{
+		if (!FParse::Command(&Cmd, TEXT("ASC")))
+		{
+			return false;
+		}
+
+		if (FParse::Command(&Cmd, TEXT("REFUND")))
+		{
+			TCHAR TransIdStr[256];
+			if (FParse::Token(Cmd, TransIdStr, UE_ARRAY_COUNT(TransIdStr), true))
+			{
+				uint64 TransactionId = FCString::Atoi64(TransIdStr);
+				[StoreKit2 testRefund:TransactionId inScene:[IOSAppDelegate GetDelegate].window.windowScene];
+			}
+
+			return true;
+		}
+		
+		return false;
+	}
+	
+} GStoreExec;
+
+#endif
+
 #endif //UE_WITH_STORE_KIT
diff --git a/Engine/Plugins/Online/IOS/OnlineSubsystemIOS/Source/Private/StoreKit2.swift b/Engine/Plugins/Online/IOS/OnlineSubsystemIOS/Source/Private/StoreKit2.swift
new file mode 100644
--- /dev/null
+++ b/Engine/Plugins/Online/IOS/OnlineSubsystemIOS/Source/Private/StoreKit2.swift
@@ -0,0 +1,17 @@
+// Copyright Epic Games, Inc. All Rights Reserved.
+
+import Foundation
+import StoreKit
+
+@objcMembers class StoreKit2 : NSObject
+{
+	static func testRefund(_ transactionId: UInt64, inScene: UIWindowScene)
+	{
+		Task
+		{
+			let status = try await Transaction.beginRefundRequest(for: transactionId, in:inScene)
+			print("Refund status: \(status)")
+		}
+	}
+}
+
diff --git a/Engine/Plugins/Online/IOS/OnlineSubsystemIOS/Source/Private/StoreKitSwiftInterop.h b/Engine/Plugins/Online/IOS/OnlineSubsystemIOS/Source/Private/StoreKitSwiftInterop.h
new file mode 100644
--- /dev/null
+++ b/Engine/Plugins/Online/IOS/OnlineSubsystemIOS/Source/Private/StoreKitSwiftInterop.h
@@ -0,0 +1,4 @@
+// Copyright Epic Games, Inc. All Rights Reserved.
+
+#pragma once
+
diff --git a/Engine/Source/Programs/UnrealBuildTool/ToolChain/AppleToolChain.cs b/Engine/Source/Programs/UnrealBuildTool/ToolChain/AppleToolChain.cs
--- a/Engine/Source/Programs/UnrealBuildTool/ToolChain/AppleToolChain.cs
+++ b/Engine/Source/Programs/UnrealBuildTool/ToolChain/AppleToolChain.cs
@@ -633,8 +633,12 @@ namespace UnrealBuildTool
 			Arguments.Add($"-sdk {ToolChainSettings.Value.GetSDKPath(CompileEnvironment.Architecture)}");
 
 			Arguments.Add($"-swift-version 5");
+			if (CompileEnvironment.bCreateDebugInfo)
+			{
+				Arguments.Add($"-g");
+			}
 
-			Arguments.Add($"-module-name {ModuleName}");
+			Arguments.Add($"-module-name {ModuleName}Swift");
 
 			if (CompileEnvironment.ForceIncludeFiles.Count() > 0)
 			{
