OnlineSubsystemIOS - Handle race condition when purchasing an offer right after finalizing a successful purchase of the same offer

diff --git a/Engine/Plugins/Online/IOS/OnlineSubsystemIOS/Source/Private/OnlinePurchaseIOS.cpp b/Engine/Plugins/Online/IOS/OnlineSubsystemIOS/Source/Private/OnlinePurchaseIOS.cpp
--- a/Engine/Plugins/Online/IOS/OnlineSubsystemIOS/Source/Private/OnlinePurchaseIOS.cpp
+++ b/Engine/Plugins/Online/IOS/OnlineSubsystemIOS/Source/Private/OnlinePurchaseIOS.cpp
@@ -200,6 +200,17 @@ struct FOnlinePurchaseInProgressTransactionIOS
 -(void)paymentQueue: (SKPaymentQueue*)Queue removedTransactions: (NSArray*)Transactions
 {
     UE_LOG_ONLINE_PURCHASE(Log, TEXT("FStoreKitPurchaseProxy::removedTransaction"));
+	
+	TSharedPtr<TArray<FOnlinePurchaseTransactionIOS>> IosTransactions = [self TranslateTransactions: Transactions];
+	
+	[FIOSAsyncTask CreateTaskWithBlock : ^ bool(void)
+	{
+		if(_PurchaseReceiver)
+		{
+			_PurchaseReceiver->OnRemovedTransactions(IosTransactions ? *IosTransactions : TArray<FOnlinePurchaseTransactionIOS>());
+		}
+		return true;
+	}];
 }
 
 -(void)paymentQueue: (SKPaymentQueue*)Queue restoreCompletedTransactionsFailedWithError: (NSError*)Error
@@ -483,7 +494,30 @@ void FOnlinePurchaseIOS::Checkout(const FUniqueNetId& UserId, const FPurchaseChe
                         bAskToBuy = OnlinePurchaseIOSPrivate::CVarSimluateAskToBuy.GetValueOnGameThread() == 1;
 #endif
 						InProgressTransaction = MakeShared<const FOnlinePurchaseInProgressTransactionIOS>(CheckoutRequest, Delegate);
-                        [StoreKitProxy MakePurchase:Product withQuantity:Offer.Quantity simulateAskToBuy:bAskToBuy];
+							
+						const bool bOfferFinalizing = FinalizingTransactions.ContainsByPredicate(
+							[OfferId = Offer.OfferId](const FKnownTransaction& Transaction)
+							{
+								return OfferId == FString(static_cast<SKPaymentTransaction*>(Transaction.PaymentTransaction).payment.productIdentifier);
+							});
+						
+						if (bOfferFinalizing)
+						{
+							UE_LOG_ONLINE_PURCHASE(Log, TEXT("Deferring purchase of offer %s until finalization of previous purchase completes."), *Offer.OfferId);
+							OnTransactionFinalized.Add(FOnTransactionFinalized::FDelegate::CreateSPLambda(InProgressTransaction.ToSharedRef(),
+								[this, ProductToPurchase = TRetainedObjCInstance<SKProduct*>(Product), OfferId = Offer.OfferId, Quantity = Offer.Quantity, bAskToBuy](const FOnlinePurchaseTransactionIOS& Transaction)
+								{
+									if (Transaction.GetOfferId() == OfferId)
+									{
+										[StoreKitProxy MakePurchase:ProductToPurchase withQuantity:Quantity simulateAskToBuy:bAskToBuy];
+										OnTransactionFinalized.RemoveAll(InProgressTransaction.Get());
+									}
+								}));
+						}
+						else
+						{
+							[StoreKitProxy MakePurchase:Product withQuantity:Offer.Quantity simulateAskToBuy:bAskToBuy];
+						}
 						bStarted = true;
                     }
                 }
@@ -530,15 +564,21 @@ void FOnlinePurchaseIOS::FinalizePurchase(const FUniqueNetId& UserId, const FStr
 {
 	UE_LOG_ONLINE_PURCHASE(Log, TEXT("FOnlinePurchaseIOS::FinalizePurchase %s %s"), *UserId.ToString(), *ReceiptId);
 
-    const FKnownTransaction* TransactionToFinalize = CachedReceipts.FindByPredicate([&ReceiptId](const FKnownTransaction& Entry) { return Entry.Receipt->TransactionId == ReceiptId; });
-    if(TransactionToFinalize)
-    {
-        // Restored transactions are already marked finalized when received but kept as known transactions so user code can process them
-        if(TransactionToFinalize->Receipt->TransactionState != EPurchaseTransactionState::Restored)
-        {
-            [StoreKitProxy FinalizeTransaction: TransactionToFinalize->PaymentTransaction];
-        }
-    }
+	const int32 TransactionToFinalizeIndex = CachedReceipts.IndexOfByPredicate([&ReceiptId](const FKnownTransaction& Entry) { return Entry.Receipt->TransactionId == ReceiptId; });
+	if (TransactionToFinalizeIndex != INDEX_NONE)
+	{
+		// Restored transactions are already marked finalized when received but kept as known transactions so user code can process them
+		if (CachedReceipts[TransactionToFinalizeIndex].Receipt->TransactionState != EPurchaseTransactionState::Restored)
+		{
+			[StoreKitProxy FinalizeTransaction: CachedReceipts[TransactionToFinalizeIndex].PaymentTransaction];
+
+			// Keep track of transactions being finalized. This must complete before attempting to purchase the same offer again
+			FinalizingTransactions.Add(MoveTemp(CachedReceipts[TransactionToFinalizeIndex]));
+		}
+		
+		// Once finalized, remove from the cached receipts as a finalized transaction will not appear in the result of a query
+		CachedReceipts.RemoveAt(TransactionToFinalizeIndex);
+	}
 	else
 	{
 		UE_LOG_ONLINE_PURCHASE(Warning, TEXT("FOnlinePurchaseIOS::FinalizePurchase transaction not found in CachedReceipts"));
@@ -687,6 +727,24 @@ void FOnlinePurchaseIOS::OnUpdatedTransactions(const TArray<FOnlinePurchaseTrans
 	}
 }
 
+void FOnlinePurchaseIOS::OnRemovedTransactions(const TArray<FOnlinePurchaseTransactionIOS>& Transactions)
+{
+	for (const FOnlinePurchaseTransactionIOS& Transaction : Transactions)
+	{
+		int32 FinalizingTransactionIndex = FinalizingTransactions.IndexOfByPredicate(
+			[&Transaction](const FKnownTransaction& KnownTransaction) 
+			{
+				return KnownTransaction.Receipt->TransactionId == Transaction.GetTransactionIdentifier();
+			});
+		
+		if (FinalizingTransactionIndex != INDEX_NONE)
+		{
+			FinalizingTransactions.RemoveAt(FinalizingTransactionIndex);
+			OnTransactionFinalized.Broadcast(Transaction);
+		}
+	}
+}
+
 void FOnlinePurchaseIOS::AddReceiptToCache(TArray<FKnownTransaction>& Cache, const TSharedRef<FPurchaseReceipt>& Receipt, const FOnlinePurchaseTransactionIOS& Transaction)
 {
     FKnownTransaction* CachedTransaction = Cache.FindByPredicate([&Transaction](const FKnownTransaction& KnownTransaction) {
diff --git a/Engine/Plugins/Online/IOS/OnlineSubsystemIOS/Source/Public/OnlinePurchaseIOS.h b/Engine/Plugins/Online/IOS/OnlineSubsystemIOS/Source/Public/OnlinePurchaseIOS.h
--- a/Engine/Plugins/Online/IOS/OnlineSubsystemIOS/Source/Public/OnlinePurchaseIOS.h
+++ b/Engine/Plugins/Online/IOS/OnlineSubsystemIOS/Source/Public/OnlinePurchaseIOS.h
@@ -58,6 +58,11 @@ public:
 	 */
 	void OnUpdatedTransactions(const TArray<FOnlinePurchaseTransactionIOS>& Transactions);
 	
+	/**
+	 * Method used internally by FStoreKitPurchaseProxy to handle finalized transactions. Not meant to be called by user code
+	 */
+	void OnRemovedTransactions(const TArray<FOnlinePurchaseTransactionIOS>& Transactions);
+	
     /**
      * Method used internally by FStoreKitPurchaseProxy to notify that transactions were restored. Not meant to be called by user code
      */
@@ -110,10 +115,17 @@ private:
     TSharedPtr<const FOnlinePurchaseInProgressTransactionIOS> InProgressTransaction;
 	
 	/** Cache of known transactions */
-	TArray< FKnownTransaction > CachedReceipts;
+	TArray<FKnownTransaction> CachedReceipts;
+	
+	/** Transactions in the process of being finalized */
+	TArray<FKnownTransaction> FinalizingTransactions;
 
 	/** Reference to the parent subsystem */
 	FOnlineSubsystemIOS* Subsystem;
+	
+	/** Delegate fired when a transaction is finalized */
+	using FOnTransactionFinalized = TMulticastDelegate<void(const FOnlinePurchaseTransactionIOS&)>;
+	FOnTransactionFinalized OnTransactionFinalized;
 };
 
 typedef TSharedPtr<FOnlinePurchaseIOS, ESPMode::ThreadSafe> FOnlinePurchaseIOSPtr;
