Store ID missmatch interface change
Adds an overload for IsAllowedToPurchase taking a reference to a FOnlineError which can provide detail on why a user may be dissallowed.
part 6

diff --git a/Engine/Plugins/Online/IOS/OnlineSubsystemIOS/Source/Private/OnlinePurchaseIOS.cpp b/Engine/Plugins/Online/IOS/OnlineSubsystemIOS/Source/Private/OnlinePurchaseIOS.cpp
--- a/Engine/Plugins/Online/IOS/OnlineSubsystemIOS/Source/Private/OnlinePurchaseIOS.cpp
+++ b/Engine/Plugins/Online/IOS/OnlineSubsystemIOS/Source/Private/OnlinePurchaseIOS.cpp
@@ -463,6 +463,13 @@ bool FOnlinePurchaseIOS::IsAllowedToPurchase(const FUniqueNetId& /*UserId*/)
 	return bCanMakePurchases;
 }
 
+bool FOnlinePurchaseIOS::IsAllowedToPurchase(const FUniqueNetId& UserId, FOnlineError& Error)
+{
+	bool bResult = IsAllowedToPurchase(UserId);
+	Error = FOnlineError(bResult ? EOnlineErrorResult::Success : EOnlineErrorResult::Unknown);
+	return bResult;
+}
+
 void FOnlinePurchaseIOS::Checkout(const FUniqueNetId& UserId, const FPurchaseCheckoutRequest& CheckoutRequest, const FOnPurchaseCheckoutComplete& Delegate)
 {
 	bool bStarted = false;
diff --git a/Engine/Plugins/Online/IOS/OnlineSubsystemIOS/Source/Public/OnlinePurchaseIOS.h b/Engine/Plugins/Online/IOS/OnlineSubsystemIOS/Source/Public/OnlinePurchaseIOS.h
--- a/Engine/Plugins/Online/IOS/OnlineSubsystemIOS/Source/Public/OnlinePurchaseIOS.h
+++ b/Engine/Plugins/Online/IOS/OnlineSubsystemIOS/Source/Public/OnlinePurchaseIOS.h
@@ -26,6 +26,7 @@ public:
 	// IOnlinePurchase
 
 	virtual bool IsAllowedToPurchase(const FUniqueNetId& UserId) override;
+	virtual bool IsAllowedToPurchase(const FUniqueNetId& UserId, FOnlineError& Error) override;
 	virtual void Checkout(const FUniqueNetId& UserId, const FPurchaseCheckoutRequest& CheckoutRequest, const FOnPurchaseCheckoutComplete& Delegate) override;
 	virtual void Checkout(const FUniqueNetId& UserId, const FPurchaseCheckoutRequest& CheckoutRequest, const FOnPurchaseReceiptlessCheckoutComplete& Delegate) override;
 	virtual void FinalizePurchase(const FUniqueNetId& UserId, const FString& ReceiptId) override;
