Enable AppStore based IAP (Store) under TestFlight distribution as the service is available in sandbox under those distribution conditions.

#Synced-CL 48848898

diff --git a/Engine/Plugins/Online/IOS/OnlineSubsystemIOS/Source/Private/OnlineSubsystemIOS.cpp b/Engine/Plugins/Online/IOS/OnlineSubsystemIOS/Source/Private/OnlineSubsystemIOS.cpp
--- a/Engine/Plugins/Online/IOS/OnlineSubsystemIOS/Source/Private/OnlineSubsystemIOS.cpp
+++ b/Engine/Plugins/Online/IOS/OnlineSubsystemIOS/Source/Private/OnlineSubsystemIOS.cpp
@@ -172,7 +172,8 @@ bool FOnlineSubsystemIOS::Init()
 		if (IsInAppPurchasingEnabled())
 		{
 #if !PLATFORM_TVOS && !PLATFORM_VISIONOS		// tvOS & VisionOS ONLY support AppStore
-			if (FMarketplaceKitModule::GetCurrentTypeStatic() == EMarketplaceType::AppStore)
+			EMarketplaceType CurrentMarketplace = FMarketplaceKitModule::GetCurrentTypeStatic();
+			if (CurrentMarketplace == EMarketplaceType::AppStore || CurrentMarketplace == EMarketplaceType::TestFlight)
 #endif
 			{
 				StoreV2Interface = MakeShareable(new FOnlineStoreIOS(this));
