Fixed include casing mismatch when compiling ios with case sensitive on

#Synced-CL 50216544

diff --git a/Engine/Plugins/Online/EOSShared/Source/EOSShared/EOSShared.Build.cs b/Engine/Plugins/Online/EOSShared/Source/EOSShared/EOSShared.Build.cs
--- a/Engine/Plugins/Online/EOSShared/Source/EOSShared/EOSShared.Build.cs
+++ b/Engine/Plugins/Online/EOSShared/Source/EOSShared/EOSShared.Build.cs
@@ -21,6 +21,8 @@ public class EOSShared : ModuleRules
 
 	public EOSShared(ReadOnlyTargetRules Target) : base(Target)
 	{
+		bRequiresPlatformSDK = true;
+		
 		Type = ModuleType.CPlusPlus;
 
 		PublicDefinitions.Add("UE_WITH_EOS_SDK_APIVERSION_WARNINGS=" + (EnableApiVersionWarnings ? "1" : "0"));
diff --git a/Engine/Plugins/Online/IOS/OnlineSubsystemIOS/Source/OnlineSubsystemIOS.Build.cs b/Engine/Plugins/Online/IOS/OnlineSubsystemIOS/Source/OnlineSubsystemIOS.Build.cs
--- a/Engine/Plugins/Online/IOS/OnlineSubsystemIOS/Source/OnlineSubsystemIOS.Build.cs
+++ b/Engine/Plugins/Online/IOS/OnlineSubsystemIOS/Source/OnlineSubsystemIOS.Build.cs
@@ -6,7 +6,7 @@ public class OnlineSubsystemIOS : ModuleRules
 {
     public OnlineSubsystemIOS(ReadOnlyTargetRules Target) : base(Target)
     {
-        bRequiresPlatformSDK = true;
+		bRequiresPlatformSDK = true;
 
 		PCHUsage = ModuleRules.PCHUsageMode.UseExplicitOrSharedPCHs;
 		PrivateIncludePaths.AddRange( 
@@ -33,7 +33,7 @@ public class OnlineSubsystemIOS : ModuleRules
             }
             );
 
-		PublicWeakFrameworks.Add("Cloudkit");
+		PublicWeakFrameworks.Add("CloudKit");
 		SwiftInteropHeader = System.IO.Path.Combine(ModuleDirectory, "Private/StoreKitSwiftInterop.h");
 		
 		if (Target.Platform == UnrealTargetPlatform.IOS)
diff --git a/Engine/Plugins/Online/IOS/OnlineSubsystemIOS/Source/Private/OnlineLeaderboardsInterfaceIOS.h b/Engine/Plugins/Online/IOS/OnlineSubsystemIOS/Source/Private/OnlineLeaderboardsInterfaceIOS.h
--- a/Engine/Plugins/Online/IOS/OnlineSubsystemIOS/Source/Private/OnlineLeaderboardsInterfaceIOS.h
+++ b/Engine/Plugins/Online/IOS/OnlineSubsystemIOS/Source/Private/OnlineLeaderboardsInterfaceIOS.h
@@ -6,7 +6,7 @@
 #include "OnlineSubsystemIOSTypes.h"
 
 #include <GameKit/GKLeaderboard.h>
-#include <Gamekit/GKPlayer.h>
+#include <GameKit/GKPlayer.h>
 
 class FOnlineLeaderboardsIOS : public IOnlineLeaderboards
 {
diff --git a/Engine/Plugins/Online/OnlineSubsystemApple/Source/OnlineSubsystemApple.Build.cs b/Engine/Plugins/Online/OnlineSubsystemApple/Source/OnlineSubsystemApple.Build.cs
--- a/Engine/Plugins/Online/OnlineSubsystemApple/Source/OnlineSubsystemApple.Build.cs
+++ b/Engine/Plugins/Online/OnlineSubsystemApple/Source/OnlineSubsystemApple.Build.cs
@@ -9,6 +9,8 @@ public class OnlineSubsystemApple : ModuleRules
 
 	public OnlineSubsystemApple(ReadOnlyTargetRules Target) : base(Target)
     {
+		bRequiresPlatformSDK = true;
+		
 		PCHUsage = ModuleRules.PCHUsageMode.UseExplicitOrSharedPCHs;
 
 		if (Target.IsInPlatformGroup(UnrealPlatformGroup.IOS))
diff --git a/Engine/Plugins/Online/OnlineSubsystemFacebook/Source/OnlineSubsystemFacebook.Build.cs b/Engine/Plugins/Online/OnlineSubsystemFacebook/Source/OnlineSubsystemFacebook.Build.cs
--- a/Engine/Plugins/Online/OnlineSubsystemFacebook/Source/OnlineSubsystemFacebook.Build.cs
+++ b/Engine/Plugins/Online/OnlineSubsystemFacebook/Source/OnlineSubsystemFacebook.Build.cs
@@ -20,6 +20,8 @@ public class OnlineSubsystemFacebook : ModuleRules
 
 	public OnlineSubsystemFacebook(ReadOnlyTargetRules Target) : base(Target)
 	{
+		bRequiresPlatformSDK = true;
+
 		PrivateDefinitions.Add("ONLINESUBSYSTEMFACEBOOK_PACKAGE=1");
 		PCHUsage = ModuleRules.PCHUsageMode.UseExplicitOrSharedPCHs;
 
diff --git a/Engine/Plugins/Online/OnlineSubsystemGoogle/Source/OnlineSubsystemGoogle.Build.cs b/Engine/Plugins/Online/OnlineSubsystemGoogle/Source/OnlineSubsystemGoogle.Build.cs
--- a/Engine/Plugins/Online/OnlineSubsystemGoogle/Source/OnlineSubsystemGoogle.Build.cs
+++ b/Engine/Plugins/Online/OnlineSubsystemGoogle/Source/OnlineSubsystemGoogle.Build.cs
@@ -15,6 +15,8 @@ public class OnlineSubsystemGoogle : ModuleRules
 
 	public OnlineSubsystemGoogle(ReadOnlyTargetRules Target) : base(Target)
 	{
+		bRequiresPlatformSDK = true;
+
 		PrivateDefinitions.Add("ONLINESUBSYSTEMGOOGLE_PACKAGE=1");
 		PCHUsage = ModuleRules.PCHUsageMode.UseExplicitOrSharedPCHs;
 
diff --git a/Engine/Plugins/Online/WebAuth/Source/WebAuth.Build.cs b/Engine/Plugins/Online/WebAuth/Source/WebAuth.Build.cs
--- a/Engine/Plugins/Online/WebAuth/Source/WebAuth.Build.cs
+++ b/Engine/Plugins/Online/WebAuth/Source/WebAuth.Build.cs
@@ -7,6 +7,8 @@ public class WebAuth : ModuleRules
 {
 	public WebAuth(ReadOnlyTargetRules Target) : base(Target)
 	{
+		bRequiresPlatformSDK = true;
+		
 		PCHUsage = PCHUsageMode.UseExplicitOrSharedPCHs;
 
 		PrivateDependencyModuleNames.AddRange(
diff --git a/Engine/Plugins/Runtime/Analytics/Adjust/Source/IOSAdjust/IOSAdjust.Build.cs b/Engine/Plugins/Runtime/Analytics/Adjust/Source/IOSAdjust/IOSAdjust.Build.cs
--- a/Engine/Plugins/Runtime/Analytics/Adjust/Source/IOSAdjust/IOSAdjust.Build.cs
+++ b/Engine/Plugins/Runtime/Analytics/Adjust/Source/IOSAdjust/IOSAdjust.Build.cs
@@ -8,6 +8,8 @@ namespace UnrealBuildTool.Rules
 	{
 		public IOSAdjust(ReadOnlyTargetRules Target) : base(Target)
         {
+			bRequiresPlatformSDK = true;
+		
 			BinariesSubFolder = "NotForLicensees";
 
 			PublicDependencyModuleNames.AddRange(
diff --git a/Engine/Plugins/Runtime/IOSReplayKit/Source/IOSReplayKit/IOSReplayKit.Build.cs b/Engine/Plugins/Runtime/IOSReplayKit/Source/IOSReplayKit/IOSReplayKit.Build.cs
--- a/Engine/Plugins/Runtime/IOSReplayKit/Source/IOSReplayKit/IOSReplayKit.Build.cs
+++ b/Engine/Plugins/Runtime/IOSReplayKit/Source/IOSReplayKit/IOSReplayKit.Build.cs
@@ -8,6 +8,8 @@ namespace UnrealBuildTool.Rules
 	{
 		public IOSReplayKit(ReadOnlyTargetRules Target) : base(Target)
 		{
+			bRequiresPlatformSDK = true;
+
 			PrivateDependencyModuleNames.AddRange(new string[]
 				{
 					"Core",
diff --git a/Engine/Plugins/Runtime/LocationServicesIOSImpl/Source/LocationServicesIOSImpl/LocationServicesIOSImpl.Build.cs b/Engine/Plugins/Runtime/LocationServicesIOSImpl/Source/LocationServicesIOSImpl/LocationServicesIOSImpl.Build.cs
--- a/Engine/Plugins/Runtime/LocationServicesIOSImpl/Source/LocationServicesIOSImpl/LocationServicesIOSImpl.Build.cs
+++ b/Engine/Plugins/Runtime/LocationServicesIOSImpl/Source/LocationServicesIOSImpl/LocationServicesIOSImpl.Build.cs
@@ -8,6 +8,8 @@ namespace UnrealBuildTool.Rules
 	{
 		public LocationServicesIOSImpl(ReadOnlyTargetRules Target) : base(Target)
 		{
+			bRequiresPlatformSDK = true;
+			
 			PublicIncludePathModuleNames.AddRange
 			(
 				new string[]
diff --git a/Engine/Source/Runtime/Advertising/IOS/IOSAdvertising/IOSAdvertising.Build.cs b/Engine/Source/Runtime/Advertising/IOS/IOSAdvertising/IOSAdvertising.Build.cs
--- a/Engine/Source/Runtime/Advertising/IOS/IOSAdvertising/IOSAdvertising.Build.cs
+++ b/Engine/Source/Runtime/Advertising/IOS/IOSAdvertising/IOSAdvertising.Build.cs
@@ -5,7 +5,9 @@ namespace UnrealBuildTool.Rules
 	public class IOSAdvertising : ModuleRules
 	{
 		public IOSAdvertising( ReadOnlyTargetRules Target ) : base(Target)
-		{			
+		{
+			bRequiresPlatformSDK = true;
+
 			PublicDependencyModuleNames.AddRange(
 				new string[]
 				{
diff --git a/Engine/Source/Runtime/Core/Private/IOS/IOSSemaphore.cpp b/Engine/Source/Runtime/Core/Private/IOS/IOSSemaphore.cpp
--- a/Engine/Source/Runtime/Core/Private/IOS/IOSSemaphore.cpp
+++ b/Engine/Source/Runtime/Core/Private/IOS/IOSSemaphore.cpp
@@ -1,3 +1,4 @@
 // Copyright Epic Games, Inc. All Rights Reserved.
 
+#include "IOS/IOSSemaphore.h"
 #include "Mac/MacSemaphore.cpp"
diff --git a/Engine/Source/Runtime/Launch/Private/IOS/LaunchIOS.cpp b/Engine/Source/Runtime/Launch/Private/IOS/LaunchIOS.cpp
--- a/Engine/Source/Runtime/Launch/Private/IOS/LaunchIOS.cpp
+++ b/Engine/Source/Runtime/Launch/Private/IOS/LaunchIOS.cpp
@@ -722,7 +722,7 @@ int main(int argc, char *argv[])
 #endif
 
 
-#include "Apple/LaunchCppToSwift.h"
+#include "Apple/LaunchCPPToSwift.h"
 
 #if PLATFORM_VISIONOS
 
diff --git a/Engine/Source/Runtime/Online/BackgroundHTTP/BackgroundHTTP.Build.cs b/Engine/Source/Runtime/Online/BackgroundHTTP/BackgroundHTTP.Build.cs
--- a/Engine/Source/Runtime/Online/BackgroundHTTP/BackgroundHTTP.Build.cs
+++ b/Engine/Source/Runtime/Online/BackgroundHTTP/BackgroundHTTP.Build.cs
@@ -6,6 +6,9 @@ public class BackgroundHTTP : ModuleRules
 {
     public BackgroundHTTP(ReadOnlyTargetRules Target) : base(Target)
     {
+		bRequiresPlatformSDK = true;
+
+
         PublicIncludePathModuleNames.AddRange(
             new string[] {
                 "HTTP",
