Add missing compiler barriers in OnlineSubsystemIOS when StoreKit support is disabled

#Synced-CL 50591631

diff --git a/Engine/Plugins/Online/IOS/OnlineSubsystemIOS/Source/Private/OnlineSubsystemIOS.cpp b/Engine/Plugins/Online/IOS/OnlineSubsystemIOS/Source/Private/OnlineSubsystemIOS.cpp
--- a/Engine/Plugins/Online/IOS/OnlineSubsystemIOS/Source/Private/OnlineSubsystemIOS.cpp
+++ b/Engine/Plugins/Online/IOS/OnlineSubsystemIOS/Source/Private/OnlineSubsystemIOS.cpp
@@ -83,12 +83,20 @@ IOnlineEntitlementsPtr FOnlineSubsystemIOS::GetEntitlementsInterface() const
 
 IOnlineStoreV2Ptr FOnlineSubsystemIOS::GetStoreV2Interface() const
 {
+#if UE_WITH_STORE_KIT
 	return StoreV2Interface;
+#else
+	return nullptr;
+#endif
 }
 
 IOnlinePurchasePtr FOnlineSubsystemIOS::GetPurchaseInterface() const
 {
+#if UE_WITH_STORE_KIT
 	return PurchaseInterface;
+#else
+	return nullptr;
+#endif
 }
 
 IOnlineEventsPtr FOnlineSubsystemIOS::GetEventsInterface() const
@@ -259,8 +267,10 @@ bool FOnlineSubsystemIOS::Shutdown()
 	DESTRUCT_INTERFACE(TurnBasedInterface);
 	DESTRUCT_INTERFACE(UserCloudInterface);
 	DESTRUCT_INTERFACE(SharedCloudInterface);
+#if UE_WITH_STORE_KIT
 	DESTRUCT_INTERFACE(StoreV2Interface);
 	DESTRUCT_INTERFACE(PurchaseInterface);
+#endif
 
 #undef DESTRUCT_INTERFACE
 	
@@ -297,11 +307,13 @@ bool FOnlineSubsystemIOS::HandlePurchaseExecCommands(UWorld* InWorld, const TCHA
 {
 	bool bWasHandled = false;
 
+#if UE_WITH_STORE_KIT
 	if (FParse::Command(&Cmd, TEXT("DUMPAPPRECEIPT")))
 	{
         PurchaseInterface->DumpAppReceipt();
+		bWasHandled = true;
 	}
-
+#endif
 	return bWasHandled;
 }
 
diff --git a/Engine/Plugins/Online/IOS/OnlineSubsystemIOS/Source/Public/OnlineSubsystemIOS.h b/Engine/Plugins/Online/IOS/OnlineSubsystemIOS/Source/Public/OnlineSubsystemIOS.h
--- a/Engine/Plugins/Online/IOS/OnlineSubsystemIOS/Source/Public/OnlineSubsystemIOS.h
+++ b/Engine/Plugins/Online/IOS/OnlineSubsystemIOS/Source/Public/OnlineSubsystemIOS.h
@@ -127,12 +127,13 @@ private:
 	/** Interface to the profile information */
 	FOnlineLeaderboardsIOSPtr LeaderboardsInterface;
 	
+#if UE_WITH_STORE_KIT
 	/** Interface to the online catalog */
 	FOnlineStoreIOSPtr StoreV2Interface;
 	
 	/** Interface to the store purchasing */
 	FOnlinePurchaseIOSPtr PurchaseInterface;
-
+#endif
 	/** Interface to the online achievements */
 	FOnlineAchievementsIOSPtr AchievementsInterface;
 
